Added GameMenuBuilder module.

diff --git a/Engine/Source/Runtime/GameMenuBuilder/GameMenuBuilder.Build.cs b/Engine/Source/Runtime/GameMenuBuilder/GameMenuBuilder.Build.cs
new file mode 100644
--- /dev/null
+++ b/Engine/Source/Runtime/GameMenuBuilder/GameMenuBuilder.Build.cs
@@ -0,0 +1,26 @@
+// Copyright 1998-2014 Epic Games, Inc. All Rights Reserved.
+
+using UnrealBuildTool;
+
+public class GameMenuBuilder : ModuleRules
+{
+    public GameMenuBuilder(TargetInfo Target)
+	{
+		PublicDependencyModuleNames.AddRange(
+			new string[] { 
+					"Engine",
+					"Core",
+                    "CoreUObject",
+					"InputCore",
+					"Slate",
+                    "SlateCore",
+			}
+		);
+
+		PrivateIncludePaths.AddRange(
+			new string[] {
+				"Runtime/GameMenuBuilder/Private",
+			}
+		);		
+	}
+}
diff --git a/Engine/Source/Runtime/GameMenuBuilder/Private/GameMenuBuilder.cpp b/Engine/Source/Runtime/GameMenuBuilder/Private/GameMenuBuilder.cpp
new file mode 100644
--- /dev/null
+++ b/Engine/Source/Runtime/GameMenuBuilder/Private/GameMenuBuilder.cpp
@@ -0,0 +1,10 @@
+// Copyright 1998-2014 Epic Games, Inc. All Rights Reserved.
+
+#include "GameMenuBuilderPrivatePCH.h"
+
+
+
+
+IMPLEMENT_MODULE(FGameMenuBuilderModule, GameMenuBuilder)
+
+
diff --git a/Engine/Source/Runtime/GameMenuBuilder/Private/GameMenuBuilderPrivatePCH.h b/Engine/Source/Runtime/GameMenuBuilder/Private/GameMenuBuilderPrivatePCH.h
new file mode 100644
--- /dev/null
+++ b/Engine/Source/Runtime/GameMenuBuilder/Private/GameMenuBuilderPrivatePCH.h
@@ -0,0 +1,10 @@
+// Copyright 1998-2014 Epic Games, Inc. All Rights Reserved.
+
+
+#include "Slate.h"
+#include "SlateGameResources.h"
+#include "Engine.h"
+
+#include "GameMenuBuilder.h"
+#include "GameMenuWidgetStyle.h"
+
diff --git a/Engine/Source/Runtime/GameMenuBuilder/Private/GameMenuBuilderStyle.cpp b/Engine/Source/Runtime/GameMenuBuilder/Private/GameMenuBuilderStyle.cpp
new file mode 100644
--- /dev/null
+++ b/Engine/Source/Runtime/GameMenuBuilder/Private/GameMenuBuilderStyle.cpp
@@ -0,0 +1,74 @@
+// Copyright 1998-2014 Epic Games, Inc. All Rights Reserved.
+
+#include "GameMenuBuilderPrivatePCH.h"
+
+
+
+TSharedPtr< FSlateStyleSet > FGameMenuBuilderStyle::SimpleStyleInstance = NULL;
+
+void FGameMenuBuilderStyle::Initialize(const FString StyleName)
+{
+	if (FModuleManager::Get().IsModuleLoaded("GameMenuBuilder") == false)
+	{
+		FModuleManager::LoadModuleChecked<FGameMenuBuilderModule>("GameMenuBuilder");
+	}
+	if (!SimpleStyleInstance.IsValid())
+	{
+		SimpleStyleInstance = Create(StyleName);
+		FSlateStyleRegistry::RegisterSlateStyle(*SimpleStyleInstance);
+	}
+}
+
+void FGameMenuBuilderStyle::Shutdown()
+{
+	FSlateStyleRegistry::UnRegisterSlateStyle(*SimpleStyleInstance);
+	ensure(SimpleStyleInstance.IsUnique());
+	SimpleStyleInstance.Reset();
+}
+
+FName FGameMenuBuilderStyle::GetStyleSetName()
+{
+	static FName StyleSetName(TEXT("MenuPageStyle"));
+	return StyleSetName;
+}
+
+#define IMAGE_BRUSH( RelativePath, ... ) FSlateImageBrush( FPaths::GameContentDir() / "Slate"/ RelativePath + TEXT(".png"), __VA_ARGS__ )
+#define BOX_BRUSH( RelativePath, ... ) FSlateBoxBrush( FPaths::GameContentDir() / "Slate"/ RelativePath + TEXT(".png"), __VA_ARGS__ )
+#define BORDER_BRUSH( RelativePath, ... ) FSlateBorderBrush( FPaths::GameContentDir() / "Slate"/ RelativePath + TEXT(".png"), __VA_ARGS__ )
+#define TTF_FONT( RelativePath, ... ) FSlateFontInfo( FPaths::GameContentDir() / "Slate"/ RelativePath + TEXT(".ttf"), __VA_ARGS__ )
+#define OTF_FONT( RelativePath, ... ) FSlateFontInfo( FPaths::GameContentDir() / "Slate"/ RelativePath + TEXT(".otf"), __VA_ARGS__ )
+
+TSharedRef< FSlateStyleSet > FGameMenuBuilderStyle::Create(const FString StyleName)
+{
+	TSharedRef<FSlateGameResources> StyleRef = FSlateGameResources::New(FGameMenuBuilderStyle::GetStyleSetName(), *StyleName);
+	
+	FSlateStyleSet& Style = StyleRef.Get();
+
+	// Fonts still need to be specified in code for now
+	Style.Set("GameMenuStyle.MenuTextStyle", FTextBlockStyle()
+		.SetFont(TTF_FONT("Fonts/Roboto-Light", 42))
+		.SetColorAndOpacity(FLinearColor::White)
+		);
+	Style.Set("GameMenuStyle.MenuHeaderTextStyle", FTextBlockStyle()
+		.SetFont(TTF_FONT("Fonts/Roboto-Light", 42))
+		.SetColorAndOpacity(FLinearColor::White)
+		);
+
+	return StyleRef;
+}
+
+#undef IMAGE_BRUSH
+#undef BOX_BRUSH
+#undef BORDER_BRUSH
+#undef TTF_FONT
+#undef OTF_FONT
+
+void FGameMenuBuilderStyle::ReloadTextures()
+{
+	FSlateApplication::Get().GetRenderer()->ReloadTextureResources();
+}
+
+const ISlateStyle& FGameMenuBuilderStyle::Get()
+{
+	return *SimpleStyleInstance;
+}
diff --git a/Engine/Source/Runtime/GameMenuBuilder/Private/GameMenuItem.cpp b/Engine/Source/Runtime/GameMenuBuilder/Private/GameMenuItem.cpp
new file mode 100644
--- /dev/null
+++ b/Engine/Source/Runtime/GameMenuBuilder/Private/GameMenuItem.cpp
@@ -0,0 +1,3 @@
+// Copyright 1998-2014 Epic Games, Inc. All Rights Reserved.
+
+#include "GameMenuBuilderPrivatePCH.h"
\ No newline at end of file
diff --git a/Engine/Source/Runtime/GameMenuBuilder/Private/GameMenuPage.cpp b/Engine/Source/Runtime/GameMenuBuilder/Private/GameMenuPage.cpp
new file mode 100644
--- /dev/null
+++ b/Engine/Source/Runtime/GameMenuBuilder/Private/GameMenuPage.cpp
@@ -0,0 +1,103 @@
+// Copyright 1998-2014 Epic Games, Inc. All Rights Reserved.
+
+#include "GameMenuBuilderPrivatePCH.h"
+
+FGameMenuPage::FGameMenuPage()
+{
+	SelectedIndex = INDEX_NONE;
+}
+
+
+TSharedRef<FGameMenuItem> FGameMenuPage::AddMenuItem(const FText& Text, TSharedPtr<FGameMenuPage> InSubMenu)
+{
+	TSharedRef<FGameMenuItem> Item = MakeShareable(new FGameMenuItem(Text));
+	Item->SubMenu = InSubMenu; 
+	MenuItems.Add(Item);
+	return Item;
+}
+
+
+void FGameMenuPage::ShowRootMenu()
+{
+	RootMenuPageWidget->BuildAndShowMenu(this->AsShared());
+	// Select the first item in the menu
+	RootMenuPageWidget->SelectItem(0);
+}
+
+bool FGameMenuPage::InitialiseRootMenu(APlayerController* InPCOwner, const FString& InStyleName, UGameViewportClient* InGameViewport)
+{
+	if (FModuleManager::Get().IsModuleLoaded("GameMenuBuilder") == false)
+	{
+		FModuleManager::LoadModuleChecked<FGameMenuBuilderModule>("GameMenuBuilder");
+	}	
+
+	TArray<FString> SplitString;
+	InStyleName.ParseIntoArray(&SplitString, TEXT("/"), true);
+	FString StylePath = InStyleName.Replace(*SplitString[SplitString.Num() - 1],TEXT(""));
+	StylePath.RemoveFromEnd("/");
+	FGameMenuBuilderStyle::Initialize(StylePath);
+
+	FString InStyleName2(TEXT("/Game/UI/Styles/PlatformerMenuStyle"));
+	bool bResult = false;
+	if (InGameViewport)
+	{
+		SAssignNew(RootMenuPageWidget, SGameMenuPageWidget)
+			.MenuStyle(&FGameMenuBuilderStyle::Get().GetWidgetStyle<FGameMenuStyle>(*InStyleName))
+			.Cursor(EMouseCursor::Default)
+			.PCOwner(TWeakObjectPtr<APlayerController>(InPCOwner));
+
+		InGameViewport->AddViewportWidgetContent(
+			SNew(SWeakWidget)
+			.PossiblyNullContent(RootMenuPageWidget));
+		PCOwner = InPCOwner;
+		RootMenuPageWidget->MyGameViewport = InGameViewport;
+		bResult = true;
+	}
+	return bResult;
+}
+
+void FGameMenuPage::DestroyRootMenu()
+{
+	if (RootMenuPageWidget.IsValid() && RootMenuPageWidget->MyGameViewport.IsValid())
+	{
+		RootMenuPageWidget->ResetMenu();		
+		RootMenuPageWidget->MyGameViewport->RemoveViewportWidgetContent(RootMenuPageWidget.ToSharedRef());
+		RootMenuPageWidget = nullptr;
+		MenuItems.Empty();
+	}
+}
+
+int32 FGameMenuPage::NumItems() const
+{
+	return MenuItems.Num();
+}
+
+TSharedPtr<FGameMenuItem> FGameMenuPage::GetItem(int32 Index) const
+{
+	check(IsValidMenuEntryIndex(Index));
+	return MenuItems[Index];
+}
+
+bool FGameMenuPage::IsValidMenuEntryIndex(int32 Index) const
+{
+	return MenuItems.IsValidIndex(Index);
+}
+
+
+void FGameMenuPage::LockControls(bool bLockState)
+{
+	if (RootMenuPageWidget.IsValid())
+	{
+		RootMenuPageWidget->LockControls(bLockState);
+	}
+}
+
+void FGameMenuPage::HideMenu()
+{
+	RootMenuPageWidget->HideMenu();
+}
+
+void FGameMenuPage::RemoveAllItems()
+{
+	MenuItems.Empty();
+}
diff --git a/Engine/Source/Runtime/GameMenuBuilder/Private/GameMenuWidgetStyle.cpp b/Engine/Source/Runtime/GameMenuBuilder/Private/GameMenuWidgetStyle.cpp
new file mode 100644
--- /dev/null
+++ b/Engine/Source/Runtime/GameMenuBuilder/Private/GameMenuWidgetStyle.cpp
@@ -0,0 +1,43 @@
+// Copyright 1998-2014 Epic Games, Inc. All Rights Reserved.
+
+#include "GameMenuBuilderPrivatePCH.h"
+
+FGameMenuStyle::FGameMenuStyle()
+{
+	LeftMarginPercent = 0.1f;
+	SubMenuMarginPercent = 0.02f;
+	AnimationSpeed = 0.2f;
+	TextColor = FLinearColor::White;
+	LayoutType = GameMenuLayoutType::Single;
+
+	TitleHorizontalAlignment = HAlign_Center;
+	TitleVerticalAlignment = VAlign_Bottom;
+	PanelsVerticalAlignment = VAlign_Center;
+}
+
+FGameMenuStyle::~FGameMenuStyle()
+{
+}
+
+const FName FGameMenuStyle::TypeName(TEXT("FGameMenuStyle"));
+
+const FGameMenuStyle& FGameMenuStyle::GetDefault()
+{
+	static FGameMenuStyle Default;
+	return Default;
+}
+
+void FGameMenuStyle::GetResources(TArray<const FSlateBrush*>& OutBrushes) const
+{
+	OutBrushes.Add(&MenuTopBrush);
+	OutBrushes.Add(&MenuFrameBrush);
+	OutBrushes.Add(&MenuBackgroundBrush);
+	OutBrushes.Add(&MenuSelectBrush);
+}
+
+
+UGameMenuWidgetStyle::UGameMenuWidgetStyle(const class FPostConstructInitializeProperties& PCIP)
+	: Super(PCIP)
+{
+	
+}
diff --git a/Engine/Source/Runtime/GameMenuBuilder/Private/SGameMenuItemWidget.cpp b/Engine/Source/Runtime/GameMenuBuilder/Private/SGameMenuItemWidget.cpp
new file mode 100644
--- /dev/null
+++ b/Engine/Source/Runtime/GameMenuBuilder/Private/SGameMenuItemWidget.cpp
@@ -0,0 +1,184 @@
+// Copyright 1998-2014 Epic Games, Inc. All Rights Reserved.
+
+#include "GameMenuBuilderPrivatePCH.h"
+
+void SGameMenuItemWidget::Construct(const FArguments& InArgs)
+{
+	MenuStyle = InArgs._MenuStyle;
+	check(MenuStyle);
+
+	PCOwner = InArgs._PCOwner;
+	Text = InArgs._Text;
+	OptionText = InArgs._OptionText;
+	OnClicked = InArgs._OnClicked;
+	OnArrowPressed = InArgs._OnArrowPressed;
+	bIsMultichoice = InArgs._bIsMultichoice;
+	bIsActiveMenuItem = false;
+	LeftArrowVisible = EVisibility::Collapsed;
+	RightArrowVisible = EVisibility::Collapsed;
+	//if attribute is set, use its value, otherwise uses default
+	InactiveTextAlpha = InArgs._InactiveTextAlpha.Get(0.5f);
+
+	ChildSlot
+	.VAlign(VAlign_Fill)
+	.HAlign(HAlign_Fill)
+	[
+		SNew(SOverlay)
+		+ SOverlay::Slot()
+		.HAlign(HAlign_Fill)
+		.VAlign(VAlign_Fill)
+		[
+			SNew(SBox)
+			.WidthOverride(800.0f)
+			.HeightOverride(68.0f)
+			[
+				SNew(SImage)
+				.ColorAndOpacity(this,&SGameMenuItemWidget::GetButtonBgColor)
+				.Image(&MenuStyle->MenuSelectBrush)
+			]
+		]
+		+SOverlay::Slot()
+		.HAlign(bIsMultichoice ? HAlign_Left : HAlign_Center)
+		.VAlign(VAlign_Center)
+		[
+			SNew(STextBlock)
+			.TextStyle(FGameMenuBuilderStyle::Get(), "GameMenuStyle.MenuTextStyle")
+			.ColorAndOpacity(this,&SGameMenuItemWidget::GetButtonTextColor)
+			.Text(Text)
+		]
+		+SOverlay::Slot()
+		.HAlign(HAlign_Right)
+		.VAlign(VAlign_Center)
+		[
+			SNew(SHorizontalBox)
+			+SHorizontalBox::Slot()
+			.AutoWidth()
+			[
+				SNew(SBorder)
+				.BorderImage(FStyleDefaults::GetNoBrush())
+				.Padding(0)
+				.Visibility(this,&SGameMenuItemWidget::GetLeftArrowVisibility)
+				.OnMouseButtonDown(this,&SGameMenuItemWidget::OnLeftArrowDown)
+				[
+					SNew(STextBlock)
+					.TextStyle(FGameMenuBuilderStyle::Get(), "GameMenuStyle.MenuTextStyle")
+					.ColorAndOpacity(this,&SGameMenuItemWidget::GetButtonTextColor)
+					.Text(FString("<"))
+				]
+			]
+			+SHorizontalBox::Slot()
+			.AutoWidth()
+			[
+				SNew(STextBlock)
+				.Visibility(bIsMultichoice ? EVisibility:: Visible : EVisibility::Collapsed )
+				.TextStyle(FGameMenuBuilderStyle::Get(), "GameMenuStyle.MenuTextStyle")
+				.ColorAndOpacity(this,&SGameMenuItemWidget::GetButtonTextColor)
+				.Text(OptionText)
+			]
+			+SHorizontalBox::Slot()
+			.AutoWidth()
+			[
+				SNew(SBorder)
+				.BorderImage(FStyleDefaults::GetNoBrush())
+				.Padding(0)
+				.Visibility(this,&SGameMenuItemWidget::GetRightArrowVisibility)
+				.OnMouseButtonDown(this,&SGameMenuItemWidget::OnRightArrowDown)
+				[
+					SNew(STextBlock)
+					.TextStyle(FGameMenuBuilderStyle::Get(), "GameMenuStyle.MenuTextStyle")
+					.ColorAndOpacity(this,&SGameMenuItemWidget::GetButtonTextColor)
+					.Text(FString(">"))
+				]
+			]
+		]
+		
+	];
+}
+
+FReply SGameMenuItemWidget::OnRightArrowDown(const FGeometry& MyGeometry, const FPointerEvent& MouseEvent)
+{
+	FReply Result = FReply::Unhandled();
+	const int32 MoveRight = 1;
+	if (OnArrowPressed.IsBound() && bIsActiveMenuItem)
+	{
+		OnArrowPressed.Execute(MoveRight);
+		Result = FReply::Handled();
+	}
+	return Result;
+}
+
+FReply SGameMenuItemWidget::OnLeftArrowDown(const FGeometry& MyGeometry, const FPointerEvent& MouseEvent)
+{
+	FReply Result = FReply::Unhandled();
+	const int32 MoveLeft = -1;
+	if (OnArrowPressed.IsBound() && bIsActiveMenuItem)
+	{
+		OnArrowPressed.Execute(MoveLeft);
+		Result = FReply::Handled();
+	}
+	return Result;
+}
+
+EVisibility SGameMenuItemWidget::GetLeftArrowVisibility() const
+{
+	return LeftArrowVisible;
+}
+
+EVisibility SGameMenuItemWidget::GetRightArrowVisibility() const
+{
+	return RightArrowVisible;
+}
+
+FSlateColor SGameMenuItemWidget::GetButtonTextColor() const
+{
+	return MenuStyle->TextColor;
+}
+
+FSlateColor SGameMenuItemWidget::GetButtonBgColor() const
+{
+	float BgAlpha = 1.0f;
+	const float MinAlpha = 0.5f;
+	const float MaxAlpha = 1.0f;
+	const float AnimSpeedModifier = 1.5f;
+	if (PCOwner.IsValid())
+	{
+		float GameTime = PCOwner->GetWorld()->GetRealTimeSeconds();
+		float AnimPercent = FMath::Abs(FMath::Sin(GameTime*AnimSpeedModifier));
+		if (bIsActiveMenuItem)
+		{
+			BgAlpha = FMath::Lerp(MinAlpha, MaxAlpha, AnimPercent);
+		}
+		else
+		{
+			BgAlpha = 0.0f;
+		}
+	}
+	return FLinearColor(1,1,1,BgAlpha);
+}
+
+FReply SGameMenuItemWidget::OnMouseButtonDown(const FGeometry& MyGeometry, const FPointerEvent& MouseEvent)
+{
+	//execute our "OnClicked" delegate, if we have one
+	if(OnClicked.IsBound() == true)
+	{
+		return OnClicked.Execute();
+	}
+
+	return FReply::Handled();
+}
+
+
+FReply SGameMenuItemWidget::OnMouseButtonUp(const FGeometry& MyGeometry, const FPointerEvent& MouseEvent)
+{
+	return FReply::Handled();
+}
+
+FReply SGameMenuItemWidget::OnMouseMove(const FGeometry& MyGeometry, const FPointerEvent& MouseEvent) 
+{
+	return FReply::Unhandled();
+}
+
+void SGameMenuItemWidget::SetMenuItemActive(bool bIsMenuItemActive)
+{
+	this->bIsActiveMenuItem = bIsMenuItemActive;
+}
diff --git a/Engine/Source/Runtime/GameMenuBuilder/Private/SGameMenuPageWidget.cpp b/Engine/Source/Runtime/GameMenuBuilder/Private/SGameMenuPageWidget.cpp
new file mode 100644
--- /dev/null
+++ b/Engine/Source/Runtime/GameMenuBuilder/Private/SGameMenuPageWidget.cpp
@@ -0,0 +1,1106 @@
+// Copyright 1998-2014 Epic Games, Inc. All Rights Reserved.
+
+#include "GameMenuBuilderPrivatePCH.h"
+#include "Engine/Console.h"
+
+
+FMenuPanel::FMenuPanel()
+{
+	CurrentState = EPanelState::Closed;
+
+	// Setup a default animation curve
+	AnimationSequence = FCurveSequence();
+	AnimationHandle = AnimationSequence.AddCurve(0, 0.2f, ECurveEaseFunction::QuadInOut);
+}
+
+void FMenuPanel::SetStyle(const FGameMenuStyle* InStyle)
+{
+	AnimationSequence = FCurveSequence();	
+	AnimationHandle = AnimationSequence.AddCurve(0, InStyle->AnimationSpeed, ECurveEaseFunction::QuadInOut);
+}
+
+void FMenuPanel::Tick(float Delta)
+{
+	if (AnimationSequence.IsPlaying() == false)
+	{
+		EPanelState::Type OldState = CurrentState;
+		switch (CurrentState)
+		{
+		case EPanelState::Opening:
+			CurrentState = EPanelState::Open;
+			break;
+		case EPanelState::Closing:
+			CurrentState = EPanelState::Closed;
+			break;
+		}
+		if (OldState != CurrentState)
+		{
+			OnStateChanged.ExecuteIfBound(CurrentState == EPanelState::Open);
+		}
+	}
+}
+
+void FMenuPanel::ClosePanel()
+{
+	if ((CurrentState != EPanelState::Closed) && (CurrentState != EPanelState::Closing))
+	{
+		if ((AnimationSequence.IsPlaying() == true) && (CurrentState == EPanelState::Opening))
+		{
+			AnimationSequence.Reverse();
+		}
+		else
+		{
+			AnimationSequence.PlayReverse();
+		}
+		CurrentState = EPanelState::Closing;
+	}
+	else
+	{
+		if (AnimationSequence.IsPlaying() == false)
+		{
+			AnimationSequence.JumpToStart();
+		}
+		if (CurrentState != EPanelState::Closing)
+		{
+			OnStateChanged.ExecuteIfBound(CurrentState == EPanelState::Open);
+		}
+	}
+}
+
+void FMenuPanel::OpenPanel()
+{
+	if ((CurrentState != EPanelState::Open) && (CurrentState != EPanelState::Opening) )
+	{
+		if ((AnimationSequence.IsPlaying() == true) && (CurrentState == EPanelState::Closing))
+		{
+			AnimationSequence.Reverse();
+		}
+		else
+		{
+			AnimationSequence.Play();
+		}
+		CurrentState = EPanelState::Opening;
+	}
+	else
+	{
+		if (AnimationSequence.IsPlaying() == false)
+		{
+			AnimationSequence.JumpToEnd();
+		}
+		if (CurrentState != EPanelState::Opening)
+		{
+			OnStateChanged.ExecuteIfBound(CurrentState == EPanelState::Open);
+		}
+	}
+}
+
+void FMenuPanel::ForcePanelOpen()
+{
+	AnimationSequence.JumpToEnd();
+	if (CurrentState != EPanelState::Open)
+	{
+		CurrentState = EPanelState::Open;
+		OnStateChanged.ExecuteIfBound(CurrentState == EPanelState::Open);
+	}
+}
+
+void FMenuPanel::ForcePanelClosed()
+{
+	AnimationSequence.JumpToStart();
+	if (CurrentState != EPanelState::Closed)
+	{
+		CurrentState = EPanelState::Closed;
+		OnStateChanged.ExecuteIfBound(CurrentState == EPanelState::Open);
+	}
+}
+
+void SGameMenuPageWidget::Construct(const FArguments& InArgs)
+{
+	MenuStyle = InArgs._MenuStyle;
+	check(MenuStyle);
+	MainMenuPanel.SetStyle(MenuStyle);
+	
+	EHorizontalAlignment MainAlignmentH;
+	if (MenuStyle->LayoutType == GameMenuLayoutType::Single)
+	{
+		MainAlignmentH = EHorizontalAlignment::HAlign_Center;
+	}
+	else
+	{
+		MainAlignmentH = EHorizontalAlignment::HAlign_Left; 
+	}
+		
+	SetupAnimations();
+	TitleWidgetAnimation.JumpToEnd();
+	
+	bControlsLocked = false;
+	bConsoleVisible = false;
+	bMenuHidden = false;
+	SelectedIndex = INDEX_NONE;
+	
+
+	PCOwner = InArgs._PCOwner;
+	bGameMenu = InArgs._GameMenu;
+	UIScale.Bind(this, &SGameMenuPageWidget::GetUIScale);
+	//ControllerHideMenuKey = EKeys::Gamepad_Special_Right;
+	Visibility.Bind(this, &SGameMenuPageWidget::GetSlateVisibility);
+		
+	// Create the title widget
+	TSharedRef<SHorizontalBox> TitleBoxWidget = SNew(SHorizontalBox)
+	+ SHorizontalBox::Slot()
+	.AutoWidth()
+	.Padding(TAttribute<FMargin>(this, &SGameMenuPageWidget::GetMenuTitleOffset))		
+	[
+		SNew(SBorder)
+		.BorderImage(&MenuStyle->MenuTopBrush)
+		.Padding(MenuStyle->TitleBorderMargin)
+		.Visibility(this, &SGameMenuPageWidget::GetMenuTitleVisibility)
+		[
+			SNew(STextBlock)
+			.TextStyle(FGameMenuBuilderStyle::Get(), "GameMenuStyle.MenuHeaderTextStyle")
+			.ColorAndOpacity(this, &SGameMenuPageWidget::GetTitleColor)
+			.Text(this, &SGameMenuPageWidget::GetMenuTitle)
+		]
+	];
+
+	// Create the widget that houses the 2 menu panels
+	TSharedRef< SHorizontalBox >	PanelBoxes = SNew(SHorizontalBox)
+	+ SHorizontalBox::Slot()
+	.HAlign(MainAlignmentH)
+	.VAlign(MenuStyle->PanelsVerticalAlignment)
+	[
+		// The main menu
+ 		SNew(SOverlay)
+		+ SOverlay::Slot()
+		.Padding(TAttribute<FMargin>(this, &SGameMenuPageWidget::GetMainMenuOffset))
+		[
+			SNew(SImage)
+			.ColorAndOpacity(this, &SGameMenuPageWidget::GetPanelsBackgroundColor)
+			.Image(&MenuStyle->MenuBackgroundBrush)
+		]	
+		+ SOverlay::Slot()
+		.Padding(TAttribute<FMargin>(this, &SGameMenuPageWidget::GetMainMenuOffset))
+		[
+			SNew(SBorder)
+			.ColorAndOpacity(this, &SGameMenuPageWidget::GetPanelsColor)
+			.BorderImage(&MenuStyle->MenuFrameBrush)
+			.Padding(TAttribute<FMargin>(this, &SGameMenuPageWidget::GetMenuItemPadding)/*MenuStyle->ItemBorderMargin*/)
+			[
+				SAssignNew(MainPanel, SVerticalBox)
+			]
+		]
+	]
+	+ SHorizontalBox::Slot()
+	.AutoWidth()
+	.HAlign(MainAlignmentH)
+	.VAlign(MenuStyle->PanelsVerticalAlignment)
+	[
+		// The sub menu
+		SNew(SOverlay)
+		.Visibility(this, &SGameMenuPageWidget::GetSubMenuVisiblity)
+		+ SOverlay::Slot()
+		.Padding(TAttribute<FMargin>(this, &SGameMenuPageWidget::GetSubMenuOffset))
+		[
+			SNew(SImage)
+			.ColorAndOpacity(this, &SGameMenuPageWidget::GetPanelsBackgroundColor)
+			.Image(&MenuStyle->MenuBackgroundBrush)
+		]
+		+ SOverlay::Slot()
+			.Padding(TAttribute<FMargin>(this, &SGameMenuPageWidget::GetSubMenuOffset))
+		[
+			SNew(SBorder)
+			.ColorAndOpacity(this, &SGameMenuPageWidget::GetPanelsColor)
+			.BorderImage(&MenuStyle->MenuFrameBrush)		
+			.Padding(TAttribute<FMargin>(this, &SGameMenuPageWidget::GetSubMenuItemPadding)/*MenuStyle->ItemBorderMargin*/)
+			[
+				SAssignNew(SubPanel, SVerticalBox)
+			]
+		]
+	];
+	
+	// Create the main widget
+	ChildSlot
+	[
+		SNew(SDPIScaler)
+		.DPIScale(UIScale)
+		[
+			SNew(SOverlay)
+			
+			+ SOverlay::Slot()
+ 			.HAlign(HAlign_Fill)
+ 			.VAlign(VAlign_Fill)
+			.Padding(TAttribute<FMargin>(this, &SGameMenuPageWidget::GetMenuOffset))
+			[
+				SNew(SVerticalBox)
+				+ SVerticalBox::Slot()
+				.HAlign(MenuStyle->TitleHorizontalAlignment)
+				.VAlign(MenuStyle->TitleVerticalAlignment)				
+				[
+					TitleBoxWidget
+				]
+				
+				+ SVerticalBox::Slot()
+				.HAlign(MainAlignmentH)
+				.VAlign(MenuStyle->PanelsVerticalAlignment)
+				.Padding(FMargin(58.0f))
+				[						
+					PanelBoxes
+				]			
+			]				//SOverlay
+		]					//SDPIScaler
+	];
+	
+}
+
+EVisibility SGameMenuPageWidget::GetSubMenuVisiblity() const
+{
+	EVisibility Visibility = EVisibility::Collapsed;
+	if ( (MenuStyle->LayoutType != GameMenuLayoutType::Single) && ( bMenuHiding == false ) )
+	{
+		Visibility = NextMenu.IsValid() ? EVisibility::Visible : EVisibility::Collapsed;
+	}
+	return Visibility;
+}
+
+EVisibility SGameMenuPageWidget::GetSlateVisibility() const
+{
+	EVisibility Visible = EVisibility::Visible;
+	if (bMenuHidden  == true)
+	{
+		Visible = EVisibility::Collapsed;
+	}
+	else if ( (bConsoleVisible == true) || (bMenuHiding == true ) )
+	{
+		Visible = EVisibility::HitTestInvisible;
+	}
+	return Visible;
+}
+
+EVisibility SGameMenuPageWidget::GetMenuTitleVisibility() const
+{
+	return CurrentMenuTitle.IsEmpty() || bMenuHidden == true ? EVisibility::Collapsed : EVisibility::Visible;
+}
+
+FString SGameMenuPageWidget::GetMenuTitle() const 
+{
+	return CurrentMenuTitle.ToString();
+}
+
+void SGameMenuPageWidget::SetCurrentMenu(TSharedPtr< class FGameMenuPage > InMenu)
+{
+	if (InMenu.IsValid())
+	{
+		if (CurrentMenu.IsValid() && CurrentMenu->RootMenuPageWidget.IsValid())
+		{
+			InMenu->RootMenuPageWidget = CurrentMenu->RootMenuPageWidget;
+		}
+		CurrentMenu = InMenu;		
+	}
+}
+
+bool SGameMenuPageWidget::SelectItem(int32 InSelection)
+{
+	bool bResult = false;
+	if (InSelection != SelectedIndex)
+	{
+		SelectionChanged(0);
+		bResult = true;
+	}
+	return bResult;
+}
+
+void SGameMenuPageWidget::BuildAndShowMenu(TSharedPtr< class FGameMenuPage > InMenu)
+{
+	SetCurrentMenu(InMenu);
+	
+	bMenuHiding = false;
+	bMenuHidden = false;
+		
+	MainMenuPanel.OnStateChanged.Unbind();
+	SubMenuPanel.OnStateChanged.Unbind();
+	OpenMainPanel(CurrentMenu);
+	MainMenuPanel.OnStateChanged.BindSP(this, &SGameMenuPageWidget::OnMainPanelStateChange);
+	SubMenuPanel.OnStateChanged.BindSP(this, &SGameMenuPageWidget::OnSubPanelStateChange);	
+	FSlateApplication::Get().PlaySound(MenuStyle->MenuEnterSound);
+}
+
+void SGameMenuPageWidget::HideMenu() 
+{
+	if (!bMenuHiding)
+	{
+		MainMenuPanel.ClosePanel();
+		SubMenuPanel.ClosePanel();
+		PendingMainMenu.Reset();
+		PendingSubMenu.Reset();
+		bMenuHiding = true;
+	}
+}
+
+void SGameMenuPageWidget::SetupAnimations()
+{
+	//Setup a curve
+	const float StartDelay = 0.0f;
+	const float SecondDelay = bGameMenu ? 0.f : 0.3f;
+	
+	if (!bGameMenu)
+	{
+		//when in main menu, only animate from left or from right side of the screen
+		MainAnimNumber = FMath::RandRange(0,1);
+	} 
+	else
+	{
+		//in game menu can use animations from top or from bottom too
+		MainAnimNumber = FMath::RandRange(0,3);
+	}
+	
+	MenuWidgetAnimation = FCurveSequence();	
+		
+	// Color fading
+	ColorFadeCurve = MenuWidgetAnimation.AddCurve(StartDelay + SecondDelay, MenuStyle->AnimationSpeed, ECurveEaseFunction::QuadInOut);
+	
+	// sliding animation
+	MenuAnimationCurve = MenuWidgetAnimation.AddCurve(StartDelay + SecondDelay, MenuStyle->AnimationSpeed, ECurveEaseFunction::QuadInOut);
+
+	// Animation for the title
+	TitleWidgetAnimation = FCurveSequence();
+	TitleWidgetCurve = TitleWidgetAnimation.AddCurve(StartDelay, MenuStyle->AnimationSpeed, ECurveEaseFunction::QuadInOut);
+}
+
+void SGameMenuPageWidget::BuildPanelButtons(TSharedPtr< class FGameMenuPage > InPanel, TSharedPtr<SVerticalBox> InBox, int32 InPreviousIndex)
+{
+	InBox->ClearChildren();
+	
+	if (InPanel.IsValid() == true)
+	{
+		for (int32 i = 0; i < InPanel->NumItems(); ++i)
+		{
+			TSharedPtr<SWidget> TmpWidget;
+			TSharedPtr<FGameMenuItem> EachItem = InPanel->GetItem(i);
+			if (EachItem.Get()->MenuItemType == EGameMenuItemType::Standard)
+			{
+				TmpWidget = SAssignNew(EachItem->Widget, SGameMenuItemWidget)
+					.MenuStyle(MenuStyle)
+					.PCOwner(PCOwner)
+					.OnClicked(this, &SGameMenuPageWidget::MouseButtonClicked, i)
+					.Text(EachItem.Get()->Text.ToString())
+					.bIsMultichoice(false);
+			}
+			else if (EachItem->MenuItemType == EGameMenuItemType::MultiChoice)
+			{
+				TmpWidget = SAssignNew(EachItem->Widget, SGameMenuItemWidget)
+					.MenuStyle(MenuStyle)
+					.PCOwner(PCOwner)
+					.OnClicked(this, &SGameMenuPageWidget::MouseButtonClicked, i)
+					.Text(EachItem.Get()->Text.ToString())
+					.bIsMultichoice(true)
+					.OnArrowPressed(this, &SGameMenuPageWidget::ChangeOption)
+					.OptionText(this, &SGameMenuPageWidget::GetOptionText, EachItem);
+				UpdateArrows(EachItem);
+			}
+			else if (EachItem.Get()->MenuItemType == EGameMenuItemType::CustomWidget)
+			{
+				TmpWidget = EachItem.Get()->CustomWidget;
+			}
+			InBox->AddSlot().HAlign(HAlign_Left).AutoHeight()
+				[
+					TmpWidget.ToSharedRef()
+				];
+		}
+		
+		if ( InPreviousIndex != INDEX_NONE )
+		{
+			SelectedIndex = InPreviousIndex;
+			TSharedPtr<FGameMenuItem> FirstMenuItem = InPanel->IsValidMenuEntryIndex(SelectedIndex) ? InPanel->GetItem(InPreviousIndex) : NULL;
+			if (FirstMenuItem.IsValid() == true )
+			{
+				if (FirstMenuItem->MenuItemType != EGameMenuItemType::CustomWidget)
+				{
+					FirstMenuItem->Widget->SetMenuItemActive(true);
+					FSlateApplication::Get().SetKeyboardFocus(SharedThis(this));
+				}
+				// If the selection has a we need to mark for pending so it will open (in side by side layout)
+				PendingSubMenu = FirstMenuItem->SubMenu;
+			}			
+		}
+	}
+
+}
+
+FString SGameMenuPageWidget::GetOptionText(TSharedPtr<FGameMenuItem> InMenuItem) const
+{
+	return InMenuItem->MultiChoice[InMenuItem->SelectedMultiChoice].ToString();
+}
+
+void SGameMenuPageWidget::UpdateArrows(TSharedPtr<FGameMenuItem> InMenuItem)
+{
+	const int32 MinIndex = InMenuItem->MinMultiChoiceIndex > -1 ? InMenuItem->MinMultiChoiceIndex : 0;
+	const int32 MaxIndex = InMenuItem->MaxMultiChoiceIndex > -1 ? InMenuItem->MaxMultiChoiceIndex : InMenuItem->MultiChoice.Num()-1;
+	const int32 CurIndex = InMenuItem->SelectedMultiChoice;
+	if (CurIndex > MinIndex)
+	{
+		InMenuItem->Widget->LeftArrowVisible = EVisibility::Visible;
+	} 
+	else
+	{
+		InMenuItem->Widget->LeftArrowVisible = EVisibility::Collapsed;
+	}
+	if (CurIndex < MaxIndex)
+	{
+		InMenuItem->Widget->RightArrowVisible = EVisibility::Visible;
+	}
+	else
+	{
+		InMenuItem->Widget->RightArrowVisible = EVisibility::Collapsed;
+	}
+}
+
+void SGameMenuPageWidget::EnterSubMenu(TSharedPtr<class FGameMenuPage> InSubMenu)
+{
+	if (InSubMenu.IsValid())
+	{
+		MenuHistory.Push(CurrentMenu);
+		CurrentMenu->SelectedIndex = SelectedIndex;
+		MainMenuPanel.ClosePanel();
+		PendingMainMenu = InSubMenu;
+		FSlateApplication::Get().PlaySound(MenuStyle->MenuEnterSound);
+	}
+}
+
+void SGameMenuPageWidget::MenuGoBack(bool bIsCancel)
+{
+	if (MenuHistory.Num() > 0)
+	{		
+		TSharedPtr< FGameMenuPage > MenuInfo = MenuHistory.Pop();
+		if (MainMenuPanel.CurrentState == EPanelState::Closing)
+		{
+			MainMenuPanel.OpenPanel();
+			PendingMainMenu.Reset();
+			PendingSubMenu.Reset();
+		}
+		else
+		{
+			CurrentMenu->SelectedIndex = SelectedIndex;
+			if (MenuStyle->LayoutType == GameMenuLayoutType::Single)
+			{
+				// single menu layout - close this panel and replace with prev menu
+				PendingMainMenu = MenuInfo;
+				MainMenuPanel.ClosePanel();
+			}
+			else
+			{
+				// side by side layout means the current menu becomes sub menu - and the main is the one we are going back too				
+				BuildPanelButtons(CurrentMenu, SubPanel, INDEX_NONE);
+				SubMenuPanel.ForcePanelOpen();
+				MainMenuPanel.ForcePanelClosed();
+				OpenMainPanel(MenuInfo);
+			}
+		}
+		FSlateApplication::Get().PlaySound(MenuStyle->MenuBackSound);
+		if (bIsCancel == true)
+		{
+			CurrentMenu->MenuGoBackCancel();
+		}
+		else
+		{
+			CurrentMenu->MenuGoBack();
+		}		
+	}
+	else if (bGameMenu) // only when it's in-game menu variant
+	{
+		if (MenuStyle->MenuBackSound.GetResourceObject() != nullptr)
+		{
+			FSlateApplication::Get().PlaySound(MenuStyle->MenuBackSound);
+		}
+		// We are sort of going back here too
+		if (bIsCancel == true)
+		{
+			CurrentMenu->MenuGoBackCancel();
+		}
+		else
+		{
+			CurrentMenu->MenuGoBack();
+		}
+		OnToggleMenu.ExecuteIfBound();
+	}
+}
+
+void SGameMenuPageWidget::ConfirmMenuItem()
+{
+	if (CurrentMenu.IsValid()==true)
+	{
+		TSharedPtr<FGameMenuItem> CurrentMenuItem = CurrentMenu->GetItem(SelectedIndex);
+		if ( CurrentMenuItem.IsValid() == true )
+		{
+			bool bItemConfirmed = false;
+			if( CurrentMenuItem->ConfirmPressed() == true)
+ 			{
+				bItemConfirmed = true;
+ 			}
+			
+			// We dont want to play 2 menu sounds here
+			if (CurrentMenuItem->SubMenu.IsValid())
+			{
+				EnterSubMenu(CurrentMenuItem->SubMenu); 
+				FSlateApplication::Get().PlaySound(MenuStyle->MenuEnterSound);				
+			}
+			else if (bItemConfirmed == true)
+			{
+				FSlateApplication::Get().PlaySound(MenuStyle->MenuItemChosenSound);
+			}
+		}
+	}
+}
+
+void SGameMenuPageWidget::OpenMainPanel(TSharedPtr< class FGameMenuPage > InMenu )
+{
+	int32 PreviousIndex = 0;
+	MainPanel->ClearChildren();
+	
+	SetCurrentMenu(InMenu);
+
+	PreviousIndex = CurrentMenu->SelectedIndex;
+	// If we have not got a previous index select the first item as that now
+	if (PreviousIndex == INDEX_NONE)
+	{
+		PreviousIndex = 0;
+	}
+	if( ( PreviousIndex != INDEX_NONE ) && (CurrentMenu->NumItems() > PreviousIndex) )
+	{
+		NextMenu = CurrentMenu->GetItem(PreviousIndex)->SubMenu;
+	}
+	BuildPanelButtons(CurrentMenu, MainPanel, PreviousIndex);
+	CurrentMenuTitle = InMenu->MenuTitle;
+
+	MainMenuPanel.OpenPanel();
+	CurrentMenu->MenuOpening();
+
+	OpenPendingSubMenu();	
+}
+
+void SGameMenuPageWidget::OnMainPanelStateChange(bool bWasOpened)
+{
+	if (bWasOpened == false)
+	{
+		if (PendingMainMenu.IsValid())
+		{
+			OpenMainPanel(PendingMainMenu);
+			if (MenuStyle->LayoutType == GameMenuLayoutType::Single)
+			{
+				OpenMainPanel(PendingMainMenu);
+			}
+			else
+			{
+				MainMenuPanel.ForcePanelOpen();
+				SubMenuPanel.ForcePanelClosed();				
+			}
+			PendingMainMenu.Reset();
+		}
+		if ((PendingSubMenu.IsValid() == false) && (bMenuHiding == true))
+		{
+			bMenuHiding = false;
+			bMenuHidden = true;
+			//Send event, if we have one bound
+			if (CurrentMenu.IsValid())
+			{
+				CurrentMenu->MenuHidden();
+			}
+		}
+		OpenPendingSubMenu();
+	}
+	else
+	{
+		OpenPendingSubMenu();
+	}
+}
+
+void SGameMenuPageWidget::OnSubPanelStateChange(bool bWasOpened)
+{
+	if (bWasOpened == false)
+	{
+		OpenPendingSubMenu();
+	}
+}
+
+void SGameMenuPageWidget::Tick(const FGeometry& AllottedGeometry, const double InCurrentTime, const float InDeltaTime)
+{
+	//Always tick the super
+	SCompoundWidget::Tick(AllottedGeometry, InCurrentTime, InDeltaTime);
+
+	//ugly code seeing if the console is open
+	UConsole* ViewportConsole = (GEngine->GameViewport != NULL) ? GEngine->GameViewport->ViewportConsole : NULL;
+	if (ViewportConsole != NULL && (ViewportConsole->ConsoleState == "Typing" || ViewportConsole->ConsoleState == "Open"))
+	{
+		if (!bConsoleVisible)
+		{
+			bConsoleVisible = true;
+			FSlateApplication::Get().SetFocusToGameViewport();
+		}
+	}
+	else
+	{
+		if (bConsoleVisible)
+		{
+			bConsoleVisible = false;
+			FSlateApplication::Get().SetKeyboardFocus(SharedThis(this));
+		}
+	}
+
+	if (GEngine && GEngine->GameViewport )
+	{
+		if (GEngine->GameViewport->ViewportFrame != NULL )
+		{
+			FViewport* Viewport = GEngine->GameViewport->ViewportFrame->GetViewport();
+			if (Viewport)
+			{
+				ScreenRes = Viewport->GetSizeXY();
+			}
+		}
+		else
+		{
+			FVector2D ViewSize;
+			GEngine->GameViewport->GetViewportSize(ViewSize);
+			ScreenRes.X = ViewSize.X;
+			ScreenRes.Y = ViewSize.Y;
+		}
+	}
+	
+
+	if (MenuWidgetAnimation.IsAtStart() && !bMenuHiding)
+	{
+		//Start the menu widget animation, set keyboard focus
+		FadeIn();
+	}
+
+	if (MenuWidgetAnimation.IsAtEnd())
+	{
+		MainMenuPanel.Tick(InDeltaTime);
+		SubMenuPanel.Tick(InDeltaTime);		
+	}	
+}
+
+float SGameMenuPageWidget::GetUIScale() const
+{
+	return ScreenRes.X / 2048.0f;
+}
+
+FMargin SGameMenuPageWidget::GetMenuOffset() const
+{
+	FMargin Result;
+	if ( CurrentMenu.IsValid() == true )
+	{
+		const float WidgetWidth = MainPanel->GetDesiredSize().X;
+		const float WidgetHeight = MainPanel->GetDesiredSize().Y;
+		const float CachedScale = UIScale.Get();
+		const float VirtualScreenWidth = ScreenRes.X / CachedScale;
+		const float VirtualScreenHeight = ScreenRes.Y / CachedScale;		
+		const float AnimProgress = MenuAnimationCurve.GetLerp();
+		const float OffsetX = 0.0f;
+
+		switch (MainAnimNumber)
+		{
+		case 0:
+			Result = FMargin(OffsetX + VirtualScreenWidth - AnimProgress * VirtualScreenWidth, 0, 0, 0);
+			break;
+		case 1:
+			Result = FMargin(OffsetX - VirtualScreenWidth + AnimProgress * VirtualScreenWidth, 0, 0, 0);
+			break;
+		case 2:
+			Result = FMargin(OffsetX, VirtualScreenHeight - AnimProgress * VirtualScreenHeight, 0, 0);
+			break;
+		case 3:
+			Result = FMargin(OffsetX, -VirtualScreenHeight + AnimProgress * VirtualScreenHeight, 0, 0);
+			break;
+		}
+	}
+	
+	return Result;
+}
+
+FMargin SGameMenuPageWidget::GetMenuTitleOffset() const
+{
+ 	const float WidgetWidth = MainPanel->GetDesiredSize().X;
+ 	const float BoxSizeX = MainPanel->GetDesiredSize().X;
+	const float RightMargin = -WidgetWidth + TitleWidgetAnimation.GetLerp() * WidgetWidth;
+	const float OutlineWidth = 2.0f;
+	return FMargin(OutlineWidth, OutlineWidth, RightMargin, OutlineWidth);
+}
+
+FMargin SGameMenuPageWidget::GetMenuItemPadding() const
+{
+	const float WidgetWidth = MainPanel->GetDesiredSize().X;
+	float Lerp = MainMenuPanel.AnimationHandle.GetLerp();
+	float WholeSize = -((WidgetWidth - (WidgetWidth*Lerp)) / 2);
+	float RealMargin = WholeSize;
+	return FMargin(RealMargin + (MenuStyle->ItemBorderMargin.Left*Lerp), MenuStyle->ItemBorderMargin.Top, RealMargin + (MenuStyle->ItemBorderMargin.Right*Lerp), MenuStyle->ItemBorderMargin.Bottom);
+}
+
+FMargin SGameMenuPageWidget::GetSubMenuItemPadding() const
+{
+	const float WidgetWidth = MainPanel->GetDesiredSize().X;
+	float Lerp = SubMenuPanel.AnimationHandle.GetLerp();
+	float WholeSize = -((WidgetWidth - (WidgetWidth*Lerp)) / 2);
+	float RealMargin = WholeSize;
+	return FMargin(RealMargin + (MenuStyle->ItemBorderMargin.Left*Lerp), MenuStyle->ItemBorderMargin.Top, RealMargin + (MenuStyle->ItemBorderMargin.Right*Lerp), MenuStyle->ItemBorderMargin.Bottom);
+}
+
+FMargin SGameMenuPageWidget::GetMainMenuOffset() const
+{
+	const float WidgetWidth = MainPanel->GetDesiredSize().X;
+	const float LeftBoxSizeX = MainPanel->GetDesiredSize().X;
+	float Lerp = MainMenuPanel.AnimationHandle.GetLerp();
+	float LeftMargin = 0.0f;
+	float RightMargin = 0.0;
+	if (MenuStyle->LayoutType == GameMenuLayoutType::Single)
+	{
+		float Size = WidgetWidth * Lerp;
+		RightMargin = LeftMargin = (WidgetWidth - Size) / 2.0f;
+	}
+	else
+	{
+		LeftMargin = WidgetWidth * MenuStyle->LeftMarginPercent;
+		RightMargin = -LeftBoxSizeX + Lerp * LeftBoxSizeX;
+	}
+	return FMargin(LeftMargin, 0, RightMargin, 0);
+}
+
+
+FMargin SGameMenuPageWidget::GetSubMenuOffset() const
+{
+	if (SubPanel.IsValid())
+	{
+		const float WidgetWidth = SubPanel->GetDesiredSize().X;
+		const float RightBoxSizeX = WidgetWidth;
+		const float LeftMargin = WidgetWidth * MenuStyle->SubMenuMarginPercent;
+		float Lerp = SubMenuPanel.AnimationHandle.GetLerp();
+		const float RightMargin = -RightBoxSizeX + Lerp * RightBoxSizeX;
+		return FMargin(LeftMargin, 0, RightMargin, 0);
+	}
+	else
+	{
+		return FMargin();
+	}
+}
+
+FSlateColor SGameMenuPageWidget::GetTitleColor() const
+{
+	FSlateColor VisibleColor = CurrentMenuTitle.IsEmpty() ? FLinearColor::Transparent : FLinearColor::White;
+	if (MenuStyle != NULL)
+	{
+		VisibleColor = GetTextColor();
+	}
+	return 	VisibleColor;
+}
+
+FLinearColor SGameMenuPageWidget::GetPanelsColor() const
+{
+	return FMath::Lerp(FLinearColor(1,1,1,0), FLinearColor(1,1,1,1), ColorFadeCurve.GetLerp());
+}
+
+FSlateColor SGameMenuPageWidget::GetPanelsBackgroundColor() const
+{
+	return FMath::Lerp(FLinearColor(1, 1, 1, 0), FLinearColor(1, 1, 1, 1), ColorFadeCurve.GetLerp());
+}
+
+FSlateColor SGameMenuPageWidget::GetTextColor() const
+{
+	return MenuStyle->TextColor;
+}
+
+FReply SGameMenuPageWidget::MouseButtonClicked(int32 SelectionIndex)
+{
+	FReply Reply = FReply::Unhandled();
+	if (CurrentMenu.IsValid() == true)
+	{
+		if (bControlsLocked)
+		{
+			Reply = FReply::Handled();
+		}
+		else 
+		{
+			if (SelectedIndex != SelectionIndex)
+			{
+				SelectionChanged(SelectionIndex);
+			}
+			ConfirmMenuItem();
+		}
+	}
+	return Reply;
+}
+
+FReply SGameMenuPageWidget::SelectionChanged(int32 SelectionIndex)
+{	
+	if (CurrentMenu.IsValid() == false)
+	{
+		return FReply::Unhandled();
+	}
+	if (bControlsLocked)
+	{
+		return FReply::Handled();
+	}
+	if (SelectedIndex != SelectionIndex)
+	{
+		TSharedPtr<FGameMenuItem> CurrentMenuItem;
+		TSharedPtr<SGameMenuItemWidget> MenuWidgetItem;
+		if (SelectedIndex != INDEX_NONE)
+		{
+			CurrentMenuItem = CurrentMenu->GetItem(SelectedIndex);
+
+			if (CurrentMenuItem.IsValid())
+			{
+				TSharedPtr<SGameMenuItemWidget> MenuWidgetItem = CurrentMenuItem->Widget;
+				MenuWidgetItem->SetMenuItemActive(false);
+			}
+		}
+		SelectedIndex = SelectionIndex;
+		TSharedPtr<FGameMenuItem> NewMenuItem = CurrentMenu->GetItem(SelectedIndex);
+		NewMenuItem->Widget->SetMenuItemActive(true);
+		FSlateApplication::Get().PlaySound(MenuStyle->MenuItemChangeSound);
+
+		OnSelectionChange.ExecuteIfBound(CurrentMenuItem, NewMenuItem);
+				
+		PendingSubMenu = NewMenuItem->SubMenu;
+		if ((PendingSubMenu.IsValid() == false) && (MenuStyle->LayoutType == GameMenuLayoutType::SideBySide))
+		{			
+			PendingSubMenu.Reset();
+		}
+		if ((SubMenuPanel.CurrentState == EPanelState::Open) || (SubMenuPanel.CurrentState == EPanelState::Opening))
+		{
+			SubMenuPanel.ClosePanel();
+		}
+		else
+		{
+			OpenPendingSubMenu();
+		}
+	}
+	else if (SelectedIndex == SelectionIndex)
+	{
+		ConfirmMenuItem();
+	}
+
+	return FReply::Handled().SetKeyboardFocus(SharedThis(this), EKeyboardFocusCause::SetDirectly);
+}
+
+void SGameMenuPageWidget::FadeIn()
+{
+	//Start the menu widget playing
+	MenuWidgetAnimation.Play();
+	SetTitleAnimation(true);
+
+	//Go into UI mode
+	FSlateApplication::Get().SetKeyboardFocus(SharedThis(this));
+}
+
+FReply SGameMenuPageWidget::OnMouseButtonDown(const FGeometry& MyGeometry, const FPointerEvent& MouseEvent)
+{
+	//If we clicked anywhere, jump to the end
+	if(MenuWidgetAnimation.IsPlaying())
+	{
+		MenuWidgetAnimation.JumpToEnd();
+	}
+
+	//Set the keyboard focus 
+	return FReply::Handled()
+	.SetKeyboardFocus(SharedThis(this), EKeyboardFocusCause::SetDirectly);
+}
+
+void SGameMenuPageWidget::ChangeOption(int32 InMoveBy)
+{
+	if (CurrentMenu.IsValid() == false)
+	{
+		return;
+	}
+	TSharedPtr<FGameMenuItem> MenuItem = CurrentMenu->GetItem(SelectedIndex);
+
+	const int32 MinIndex = MenuItem->MinMultiChoiceIndex > -1 ? MenuItem->MinMultiChoiceIndex : 0;
+	const int32 MaxIndex = MenuItem->MaxMultiChoiceIndex > -1 ? MenuItem->MaxMultiChoiceIndex : MenuItem->MultiChoice.Num()-1;
+	const int32 CurIndex = MenuItem->SelectedMultiChoice;
+
+	if (MenuItem->MenuItemType == EGameMenuItemType::MultiChoice)
+	{
+		if ( CurIndex + InMoveBy >= MinIndex && CurIndex + InMoveBy <= MaxIndex)
+		{
+			MenuItem->SelectedMultiChoice += InMoveBy;
+			MenuItem->OnOptionChanged.ExecuteIfBound(MenuItem, MenuItem->SelectedMultiChoice);
+			FSlateApplication::Get().PlaySound(MenuStyle->OptionChangeSound);
+		}
+		UpdateArrows(MenuItem);
+	}
+}
+
+FReply SGameMenuPageWidget::OnKeyDown(const FGeometry& MyGeometry, const FKeyboardEvent& InKeyboardEvent)
+{
+	FReply Result = FReply::Unhandled();
+	
+	if ( (CurrentMenu.IsValid() == true ) && ( !bControlsLocked) )
+	{
+		bool bNavigationLocked = bControlsLocked | PendingMainMenu.IsValid() | PendingSubMenu.IsValid();
+		const FKey Key = InKeyboardEvent.GetKey();
+	
+		if (bNavigationLocked == false )
+		{
+			if (Key == EKeys::Up || Key == EKeys::Down)
+			{
+				int32 MoveBy = Key == EKeys::Up ? -1 : 1;
+
+				if (SelectedIndex + MoveBy > -1 && SelectedIndex + MoveBy < CurrentMenu->NumItems())
+				{
+					SelectionChanged(SelectedIndex + MoveBy);
+				}
+				Result = FReply::Handled();
+			}
+			else if (Key == EKeys::Left || Key == EKeys::Right)
+			{
+				int32 MoveBy = Key == EKeys::Left ? -1 : 1;
+				ChangeOption(MoveBy);
+				Result = FReply::Handled();
+			}
+		}
+		if (Key == EKeys::Enter)
+		{
+			ConfirmMenuItem();
+			Result = FReply::Handled();
+		} 
+		else if (Key == EKeys::Escape )
+		{
+			MenuGoBack(true);
+			Result = FReply::Handled();
+		}
+	}
+	return Result;
+}
+
+FReply SGameMenuPageWidget::OnControllerButtonPressed( const FGeometry& MyGeometry, const FControllerEvent& ControllerEvent )
+{
+	FReply Result = FReply::Unhandled();
+	if ((CurrentMenu.IsValid() == true) && (!bControlsLocked))
+	{
+		bool bNavigationLocked = bControlsLocked | PendingMainMenu.IsValid() | PendingSubMenu.IsValid();
+		const FKey Key = ControllerEvent.GetEffectingButton();
+		if (bNavigationLocked == false)
+		{
+			if (Key == EKeys::Gamepad_DPad_Up || Key == EKeys::Gamepad_LeftStick_Up)
+			{
+				if (SelectedIndex > 0)
+				{
+					SelectionChanged(SelectedIndex - 1);
+				}
+				Result = FReply::Handled();
+			}
+			else if (Key == EKeys::Gamepad_DPad_Down || Key == EKeys::Gamepad_LeftStick_Down)
+			{
+				if (SelectedIndex + 1 < CurrentMenu->NumItems())
+				{
+					SelectionChanged(SelectedIndex + 1);
+				}
+				Result = FReply::Handled();
+			}
+			else if (Key == EKeys::Gamepad_DPad_Left || Key == EKeys::Gamepad_LeftStick_Left)
+			{
+				ChangeOption(-1);
+				Result = FReply::Handled();
+			}
+			else if (Key == EKeys::Gamepad_DPad_Right || Key == EKeys::Gamepad_LeftStick_Right)
+			{
+				ChangeOption(1);
+				Result = FReply::Handled();
+			}
+		}
+		if (Key == EKeys::Gamepad_FaceButton_Bottom)
+		{
+			ConfirmMenuItem();
+			Result = FReply::Handled();
+		} 
+		else if (Key == EKeys::Gamepad_FaceButton_Right || Key == EKeys::Gamepad_Special_Left)
+		{
+			// Treat the right facepad button as CANCEL, but the menu button as simply back.
+			MenuGoBack(Key == EKeys::Gamepad_FaceButton_Right);
+			Result = FReply::Handled();
+		} 
+// 		else if (Key == ControllerHideMenuKey)
+// 		{
+// 			OnToggleMenu.ExecuteIfBound();
+// 			Result = FReply::Handled();
+// 		}
+	}
+	return Result;
+}
+
+FReply SGameMenuPageWidget::OnKeyboardFocusReceived(const FGeometry& MyGeometry, const FKeyboardFocusEvent& InKeyboardFocusEvent)
+{
+	if (CurrentMenu.IsValid() == false)
+	{
+		return FReply::Unhandled();
+	}
+
+	//Focus the custom widget
+	if (CurrentMenu.IsValid() && CurrentMenu->NumItems() == 1 && CurrentMenu->GetItem(0)->MenuItemType == EGameMenuItemType::CustomWidget)
+	{
+		return FReply::Handled().ReleaseJoystickCapture().SetKeyboardFocus(CurrentMenu->GetItem(0)->CustomWidget.ToSharedRef(),EKeyboardFocusCause::SetDirectly);
+	}
+
+	return FReply::Handled().ReleaseMouseCapture().CaptureJoystick(SharedThis( this ));
+}
+
+void SGameMenuPageWidget::SetTitleAnimation(bool bShowTitle)
+{
+	if (TitleWidgetAnimation.IsPlaying() == false)
+	{
+		if (bShowTitle == true)
+		{
+			TitleWidgetAnimation.Play();
+		}
+		else
+		{
+			TitleWidgetAnimation.PlayReverse();
+		}		
+	}
+	else
+	{
+		if (TitleWidgetAnimation.IsForward() == bShowTitle)
+		{
+			TitleWidgetAnimation.Reverse();
+		}
+		else
+		{
+			TitleWidgetAnimation.PlayReverse();
+		}
+	}
+}
+
+void SGameMenuPageWidget::LockControls(bool bEnable)
+{
+	bControlsLocked = bEnable;
+}
+
+void SGameMenuPageWidget::OpenPendingSubMenu()
+{
+	if( (PendingSubMenu.IsValid() ) && ( MenuStyle->LayoutType == GameMenuLayoutType::SideBySide ) )
+	{
+		BuildPanelButtons(PendingSubMenu, SubPanel, INDEX_NONE);		
+		SubMenuPanel.OpenPanel();
+		NextMenu = PendingSubMenu;
+		NextMenu->MenuOpening();
+		PendingSubMenu.Reset();
+	}
+	else
+	{
+		PendingSubMenu.Reset();
+	}
+}
+
+
+
+void SGameMenuPageWidget::ResetMenu()
+{
+	while (MenuHistory.Num() != 0)
+	{
+		TSharedPtr< FGameMenuPage > MenuInfo = MenuHistory.Pop();
+		MenuInfo->RemoveAllItems();
+	}
+	MainPanel->ClearChildren();
+	SubPanel->ClearChildren();
+	PendingSubMenu.Reset();
+	PendingMainMenu.Reset();
+	CurrentMenu.Reset();
+	SubPanel.Reset();	
+}
diff --git a/Engine/Source/Runtime/GameMenuBuilder/Public/GameMenuBuilder.h b/Engine/Source/Runtime/GameMenuBuilder/Public/GameMenuBuilder.h
new file mode 100644
--- /dev/null
+++ b/Engine/Source/Runtime/GameMenuBuilder/Public/GameMenuBuilder.h
@@ -0,0 +1,25 @@
+// Copyright 1998-2014 Epic Games, Inc. All Rights Reserved.
+
+#include "GameMenuBuilderModule.h"
+
+#include "Slate.h"
+#include "GameMenuBuilderStyle.h"
+#include "GameMenuWidgetStyle.h"
+#include "SGameMenuPageWidget.h"
+#include "SGameMenuItemWidget.h"
+#include "GameMenuItem.h"
+#include "GameMenuPage.h"
+
+#pragma once
+
+class FGameMenuBuilderModule : public IGameMenuBuilderModule
+{
+	virtual void StartupModule() override
+	{
+		
+	}
+	void SetStyleName()
+	{
+	}
+
+};
\ No newline at end of file
diff --git a/Engine/Source/Runtime/GameMenuBuilder/Public/GameMenuBuilderModule.h b/Engine/Source/Runtime/GameMenuBuilder/Public/GameMenuBuilderModule.h
new file mode 100644
--- /dev/null
+++ b/Engine/Source/Runtime/GameMenuBuilder/Public/GameMenuBuilderModule.h
@@ -0,0 +1,27 @@
+// Copyright 1998-2014 Epic Games, Inc. All Rights Reserved.
+
+#pragma once
+
+#include "ModuleManager.h"
+
+
+/**
+ * The public interface to this module
+ */
+class IGameMenuBuilderModule : public IModuleInterface
+{
+
+public:
+
+	/**
+	 * Singleton-like access to this module's interface.  This is just for convenience!
+	 * Beware of calling this during the shutdown phase, though.  Your module might have been unloaded already.
+	 *
+	 * @return Returns singleton instance, loading the module on demand if needed
+	 */
+	static inline IGameMenuBuilderModule& Get()
+	{
+		return FModuleManager::LoadModuleChecked< IGameMenuBuilderModule >("GameMenuBuilder");
+	}
+};
+
diff --git a/Engine/Source/Runtime/GameMenuBuilder/Public/GameMenuBuilderStyle.h b/Engine/Source/Runtime/GameMenuBuilder/Public/GameMenuBuilderStyle.h
new file mode 100644
--- /dev/null
+++ b/Engine/Source/Runtime/GameMenuBuilder/Public/GameMenuBuilderStyle.h
@@ -0,0 +1,27 @@
+// Copyright 1998-2014 Epic Games, Inc. All Rights Reserved.
+
+#pragma once
+
+class FGameMenuBuilderStyle
+{
+public:
+
+	static void Initialize(const FString StyleName);
+
+	static void Shutdown();
+	
+	/** reloads textures used by slate renderer */
+	static void ReloadTextures();
+
+	/** @return The Slate style set  */
+	static const ISlateStyle& Get();
+
+	static FName GetStyleSetName();
+		
+private:
+
+	static TSharedRef< class FSlateStyleSet > Create(const FString StyleName);
+
+	static TSharedPtr< class FSlateStyleSet > SimpleStyleInstance;
+};
+
diff --git a/Engine/Source/Runtime/GameMenuBuilder/Public/GameMenuItem.h b/Engine/Source/Runtime/GameMenuBuilder/Public/GameMenuItem.h
new file mode 100644
--- /dev/null
+++ b/Engine/Source/Runtime/GameMenuBuilder/Public/GameMenuItem.h
@@ -0,0 +1,111 @@
+// Copyright 1998-2014 Epic Games, Inc. All Rights Reserved.
+#pragma once
+
+namespace EGameMenuItemType
+{
+	enum Type
+	{
+		Root,
+		Standard,
+		MultiChoice,
+		CustomWidget,
+	};
+};
+
+class FGameMenuItem : public TSharedFromThis<FGameMenuItem>
+{
+public:
+	/** confirm menu item delegate */
+	DECLARE_DELEGATE(FOnConfirmMenuItem);
+
+	/** multi-choice option changed, parameters are menu item itself and new multi-choice index  */
+	DECLARE_DELEGATE_TwoParams(FOnOptionChanged, TSharedPtr<FGameMenuItem>, int32);
+
+	/** delegate, which is executed by SSimpleMenuWidget if user confirms this menu item */
+	FOnConfirmMenuItem OnConfirmMenuItem;
+
+	/** multi-choice option changed, parameters are menu item itself and new multi-choice index */
+	FOnOptionChanged OnOptionChanged;
+
+	/** menu item type */
+	EGameMenuItemType::Type MenuItemType;
+
+	/** menu item text */
+	FText Text;
+
+	/** sub menu if present */
+	TSharedPtr<class FGameMenuPage> SubMenu;
+
+	/** shared pointer to actual slate widget representing the menu item */
+	TSharedPtr<SGameMenuItemWidget> Widget;
+
+	/** shared pointer to actual slate widget representing the custom menu item, ie whole options screen */
+	TSharedPtr<SWidget> CustomWidget;
+
+	/** texts for multiple choice menu item (like INF AMMO ON/OFF or difficulty/resolution etc) */
+	TArray<FText> MultiChoice;
+
+	/** set to other value than -1 to limit the options range */
+	int32 MinMultiChoiceIndex;
+
+	/** set to other value than -1 to limit the options range */
+	int32 MaxMultiChoiceIndex;
+
+	/** selected multi-choice index for this menu item */
+	int32 SelectedMultiChoice;
+
+	/** true if this item is active */
+	bool bInactive;
+
+	/** constructor accepting menu item text */
+	FGameMenuItem(FText _text)
+	{
+		Text = _text;
+		MenuItemType = EGameMenuItemType::Standard;
+		MinMultiChoiceIndex = MaxMultiChoiceIndex = SelectedMultiChoice = 0;
+		bInactive = false;
+	}
+
+	/** custom widgets cannot contain sub menus, all functionality must be handled by custom widget itself */
+	FGameMenuItem(TSharedPtr<SWidget> _Widget)
+	{
+		MenuItemType = EGameMenuItemType::CustomWidget;
+		CustomWidget = _Widget;
+		MinMultiChoiceIndex = MaxMultiChoiceIndex = SelectedMultiChoice = 0;
+		bInactive = false;
+	}
+
+	/** constructor for multi-choice item */
+	FGameMenuItem(FText InText, TArray<FText> InOptions, int32 InDefaultIndex=0)
+	{
+		Text = InText;
+		MenuItemType = EGameMenuItemType::MultiChoice;
+		MultiChoice = InOptions;
+		MinMultiChoiceIndex = MaxMultiChoiceIndex = -1;
+		SelectedMultiChoice = InDefaultIndex;
+		bInactive = false;
+	}
+
+	bool ConfirmPressed()
+	{
+		bool bHandled = false;
+		if ( (bInactive == false ) && (OnConfirmMenuItem.IsBound() == true ) )
+		{
+			OnConfirmMenuItem.Execute();			
+			bHandled = true;
+		}
+		return bHandled;
+	}
+	
+	/** create special root item */
+	static TSharedRef<FGameMenuItem> CreateRoot()
+	{
+		return MakeShareable(new FGameMenuItem());
+	}
+	
+private:
+	FGameMenuItem()
+	{
+		MenuItemType = EGameMenuItemType::Root;
+	}
+};
\ No newline at end of file
diff --git a/Engine/Source/Runtime/GameMenuBuilder/Public/GameMenuPage.h b/Engine/Source/Runtime/GameMenuBuilder/Public/GameMenuPage.h
new file mode 100644
--- /dev/null
+++ b/Engine/Source/Runtime/GameMenuBuilder/Public/GameMenuPage.h
@@ -0,0 +1,202 @@
+// Copyright 1998-2014 Epic Games, Inc. All Rights Reserved.
+
+#pragma once
+
+class GAMEMENUBUILDER_API FGameMenuPage : public TSharedFromThis<FGameMenuPage>
+{
+public:
+	FGameMenuPage();
+	
+	/** Delegate for when user is going back from submenu. */
+	DECLARE_DELEGATE(FOnMenuGoBack);
+
+	/** Delegate for when menu is hidden. */
+	DECLARE_DELEGATE(FOnMenuHidden);
+	
+	/** Delegate for when menu is about to open. */
+	DECLARE_DELEGATE(FOnMenuOpening);
+
+	/* 
+	 * Initialize the menu page widget and set this menu as root. 
+	 * @param	InPCOwner		Player controller that owns this menu.
+	 * @param	InStyleName		Name of the menu style to use.
+	 * @param	InGameViewport	Gameviewport on which to place the menu.
+	 * @returns true on success. (Will fail if viewport invalid)
+	 */
+	bool InitialiseRootMenu(APlayerController* InPCOwner, const FString& InStyleName, UGameViewportClient* InGameViewport);
+
+	void DestroyRootMenu();
+
+	/* Return the number of items in menu item list. */
+	int32 NumItems() const;
+
+	/* Get the specified item in the menu item list. Will assert if the value is out of range. */
+	TSharedPtr<FGameMenuItem> GetItem(int32 Index) const;
+	
+	/* Is the supplied index with the valid range of menu items. Returns false if it is not. */
+	bool IsValidMenuEntryIndex(int32 Index) const;
+				
+	/** Weak pointer to Owning player controller. */
+	TWeakObjectPtr<APlayerController> PCOwner;
+
+	/** Builds the menu widget and shows the menu. */
+	void ShowRootMenu();
+
+	/** Current selection in this menu. */
+	int32 SelectedIndex;
+
+	/** The menu title. */
+	FText MenuTitle;
+
+	/* 
+	 * Add a menu item.
+	 *
+	 * @param	Text		The string for the item (EG START GAME)
+	 * @param	InSubMenu	Any submenu associated with the item
+	 * @returns A shared reference to the created item
+	 */
+	TSharedRef<FGameMenuItem> AddMenuItem(const FText& Text, TSharedPtr<FGameMenuPage> InSubMenu = nullptr);
+
+	/*
+	*
+	* Add a menu item.
+	*
+	* @param	Text		Label of the menu item
+	* @param	InObj		Menu page object
+	* @param	InMethod	Method to call when selection changes.
+	* @returns  SharedRef to the MenuItem that was created.
+	*/
+	template< class UserClass >
+	FORCEINLINE TSharedRef<FGameMenuItem> AddMenuItem(const FText& InText, UserClass* InObj, typename FGameMenuItem::FOnConfirmMenuItem::TSPMethodDelegate< UserClass >::FMethodPtr InMethod)
+	{
+		TSharedPtr<FGameMenuItem> Item = MakeShareable(new FGameMenuItem(InText));
+		Item->OnConfirmMenuItem.BindSP(InObj, InMethod);
+		MenuItems.Add(Item);
+		return Item.ToSharedRef();
+	}
+
+	/*
+	 * Add a menu entry with a variable number of selectable options
+	 *
+	 * @param	Text		Label of the menu item
+	 * @param	OptionsList	List of selectable items
+	 * @param	InObj		Menu page object
+	 * @param	InMethod	Method to call when selection changes.
+	 * @returns  SharedRef to the MenuItem that was created.
+	 */
+	template< class UserClass >
+	FORCEINLINE TSharedRef<FGameMenuItem> AddMenuItemWithOptions(const FText& Text, const TArray<FText>& OptionsList, UserClass* InObj, typename FGameMenuItem::FOnOptionChanged::TSPMethodDelegate< UserClass >::FMethodPtr InMethod)
+	{
+		TSharedPtr<FGameMenuItem> Item = MakeShareable(new FGameMenuItem(Text, OptionsList));
+		Item->OnOptionChanged.BindSP(InObj, InMethod);
+		MenuItems.Add(Item);
+		return Item.ToSharedRef();
+	}
+
+	/*
+	* Add a handler for the menu being canceled.
+	*
+	* @param	InObj		Menu page object
+	* @param	InMethod	Method to call when selection changes.
+	*/
+	template< class UserClass >
+	FORCEINLINE void SetCancelHandler(UserClass* InObj, typename FGameMenuPage::FOnMenuGoBack::TSPMethodDelegate< UserClass >::FMethodPtr InMethod)
+	{
+		OnGoBackCancel.BindSP(InObj, InMethod);
+	}
+	
+	/*
+	* Add a handler for the menu being canceled.
+	*
+	* @param	InObj		Menu page object
+	* @param	InMethod	Method to call when menu has been hidden.
+	*/
+	template< class UserClass >
+	FORCEINLINE void SetOnHiddenHandler(UserClass* InObj, typename FGameMenuPage::FOnMenuHidden::TSPMethodDelegate< UserClass >::FMethodPtr InMethod)
+	{
+		OnMenuHidden.BindSP(InObj, InMethod);
+	}
+
+	/*
+	* Add a handler for the menu being accepted.
+	*
+	* @param	InObj		Menu page object
+	* @param	InMethod	Method to call when selection changes.
+	*/
+	template< class UserClass >
+	FORCEINLINE void SetAcceptHandler(UserClass* InObj, typename FGameMenuPage::FOnMenuGoBack::TSPMethodDelegate< UserClass >::FMethodPtr InMethod)
+	{
+		OnGoBack.BindSP(InObj, InMethod);
+	}
+
+	/*
+	* Add a handler for the menu being opened.
+	*
+	* @param	InObj		Menu page object
+	* @param	InMethod	Method to call when menu is about to open.
+	*/
+	template< class UserClass >
+	FORCEINLINE void SetOnOpenHandler(UserClass* InObj, typename FGameMenuPage::FOnMenuHidden::TSPMethodDelegate< UserClass >::FMethodPtr InMethod)
+	{
+		OnMenuOpening.BindSP(InObj, InMethod);
+	}
+
+	/*
+	 * Lock/Unlock the menu and prevent user interaction.
+	 *
+	 * @param		bLockState	true to lock
+	 */
+	void LockControls(bool bLockState);
+		
+	/** The widget that is the menu. */
+	TSharedPtr<SGameMenuPageWidget> RootMenuPageWidget;
+
+	/** Executed when user want's to go back to previous menu. */
+	void MenuGoBack()
+	{
+		OnGoBack.ExecuteIfBound();
+	}
+
+	/** Called when user want's to CANCEL and go back to previous menu. */
+	void MenuGoBackCancel()
+	{
+		OnGoBackCancel.ExecuteIfBound();
+	}
+
+	/* Called when the menu has finished hiding. */
+	void MenuHidden()
+	{
+		OnMenuHidden.ExecuteIfBound();
+	}
+
+	/* Called when the menu is about to be opened. */
+	void MenuOpening()
+	{
+		OnMenuOpening.ExecuteIfBound();
+	}
+
+	/* Hide the menu. */
+	void HideMenu();
+
+	/* Remove all the items from the item array. */
+	void RemoveAllItems();
+
+protected:	
+	
+	/** Executed when user want's to go back to previous menu. */
+	FOnMenuGoBack OnGoBack;
+
+	/** Executed when user want's to CANCEL and go back to previous menu. */
+	FOnMenuGoBack OnGoBackCancel;
+
+	/** Delegate, which is executed when menu is finished hiding. */
+	FOnMenuHidden OnMenuHidden;
+
+	/** Delegate, which is executed when menu is about to open. */
+	FOnMenuOpening OnMenuOpening;
+
+	/* Array of menu items that represents the menu. */
+	typedef TArray< TSharedPtr<FGameMenuItem> > MenuItemList;
+	MenuItemList	MenuItems;
+};
+
diff --git a/Engine/Source/Runtime/GameMenuBuilder/Public/GameMenuWidgetStyle.h b/Engine/Source/Runtime/GameMenuBuilder/Public/GameMenuWidgetStyle.h
new file mode 100644
--- /dev/null
+++ b/Engine/Source/Runtime/GameMenuBuilder/Public/GameMenuWidgetStyle.h
@@ -0,0 +1,192 @@
+// Copyright 1998-2014 Epic Games, Inc. All Rights Reserved.
+
+#pragma once
+
+#include "SlateWidgetStyleContainerBase.h"
+#include "GameMenuWidgetStyle.generated.h"
+
+UENUM()
+namespace GameMenuLayoutType
+{
+	enum Type
+	{
+		/* Single one page menu. */
+		Single,
+		/* Side by side menu. Main menu on left, any sub menu on right. */
+		SideBySide
+	};
+}
+
+USTRUCT()
+struct FGameMenuStyle : public FSlateWidgetStyle
+{
+	GENERATED_USTRUCT_BODY()
+
+	FGameMenuStyle();
+	virtual ~FGameMenuStyle();
+
+	// FSlateWidgetStyle
+	virtual void GetResources(TArray<const FSlateBrush*>& OutBrushes) const override;
+	static const FName TypeName;
+	virtual const FName GetTypeName() const override { return TypeName; };
+	static const FGameMenuStyle& GetDefault();
+
+	/**
+	 * The brush used to draw the top of the menu
+	 */
+	UPROPERTY(EditAnywhere, Category = Brushes)
+	FSlateBrush MenuTopBrush;
+	FGameMenuStyle& SetMenuTopBrush(const FSlateBrush& InMenuTopBrush) { MenuTopBrush = InMenuTopBrush; return *this; }
+
+	/**
+	 * The brush used to draw the menu border
+	 */
+	UPROPERTY(EditAnywhere, Category = Brushes)
+	FSlateBrush MenuFrameBrush;
+	FGameMenuStyle& SetMenuFrameBrush(const FSlateBrush& InMenuFrameBrush) { MenuFrameBrush = InMenuFrameBrush; return *this; }
+
+	/**
+	* The brush used to draw the menu background
+	*/
+	UPROPERTY(EditAnywhere, Category = Brushes)
+	FSlateBrush MenuBackgroundBrush;
+	FGameMenuStyle& SetMenuBackgroundBrush(const FSlateBrush& InMenuBackgroundBrush) { MenuBackgroundBrush = InMenuBackgroundBrush; return *this; }
+
+	/**
+	 * The brush used to draw the selected menu item
+	 */
+	UPROPERTY(EditAnywhere, Category = Brushes)
+	FSlateBrush MenuSelectBrush;
+	FGameMenuStyle& SetMenuSelectBrush(const FSlateBrush& InMenuSelectBrush) { MenuSelectBrush = InMenuSelectBrush; return *this; }
+
+	/**
+	 * The sound the menu should play when entering a sub-menu
+	 */
+	UPROPERTY(EditAnywhere, Category = Sound)
+	FSlateSound MenuEnterSound;
+	FGameMenuStyle& SetMenuEnterSound(const FSlateSound& InMenuEnterSound) { MenuEnterSound = InMenuEnterSound; return *this; }
+
+	/**
+	 * The sound the menu should play when returning to a parent menu
+	 */
+	UPROPERTY(EditAnywhere, Category = Sound)
+	FSlateSound MenuBackSound;
+	FGameMenuStyle& SetMenuBackSound(const FSlateSound& InMenuBackSound) { MenuBackSound = InMenuBackSound; return *this; }
+
+	/**
+	 * The sound the menu should play when a menu option is changed
+	 */
+	UPROPERTY(EditAnywhere, Category = Sound)
+	FSlateSound OptionChangeSound;
+	FGameMenuStyle& SetOptionChangeSound(const FSlateSound& InOptionChangeSound) { OptionChangeSound = InOptionChangeSound; return *this; }
+
+	/**
+	 * The sound the menu should play when the selected menu item is changed
+	 */
+	UPROPERTY(EditAnywhere, Category = Sound)
+	FSlateSound MenuItemChangeSound;
+	FGameMenuStyle& SetMenuItemChangeSound(const FSlateSound& InMenuItemChangeSound) { MenuItemChangeSound = InMenuItemChangeSound; return *this; }
+
+	/**
+	 * The sound the menu should play when the selected menu item is disabled or unbound
+	 */
+	UPROPERTY(EditAnywhere, Category = Sound)
+	FSlateSound MenuItemInactiveSound;
+	FGameMenuStyle& SetMenuItemInactiveSound(const FSlateSound& InMenuItemInactiveSound) { MenuItemInactiveSound = InMenuItemInactiveSound; return *this; }
+
+	/**
+	 * The sound the menu should play when the selected menu item is chosen
+	 */
+	UPROPERTY(EditAnywhere, Category = Sound)
+	FSlateSound MenuItemChosenSound;
+	FGameMenuStyle& SetMenuItemChosenSound(const FSlateSound& InMenuItemChosenSound) { MenuItemChosenSound = InMenuItemChosenSound; return *this; }
+	
+	/**
+	 * The left hand margin for the main menu 
+	 */
+	UPROPERTY(EditAnywhere, Category = Layout)
+	float LeftMarginPercent;
+	FGameMenuStyle& SetMenuLeftMarginPercent(const float InLeftMarginPercent) { LeftMarginPercent = InLeftMarginPercent; return *this; }
+	
+	/**
+	* The left hand margin for the sub menu
+	*/
+	UPROPERTY(EditAnywhere, Category = Layout)
+	float SubMenuMarginPercent;
+	FGameMenuStyle& SetSubMenuMarginPercent(const float InSubMenuMarginPercent) { SubMenuMarginPercent = InSubMenuMarginPercent; return *this; }
+	
+	/**
+	* Layout type.
+	*/
+	UPROPERTY(EditAnywhere, Category = Layout)
+	TEnumAsByte<GameMenuLayoutType::Type> LayoutType;
+	FGameMenuStyle& SetMenuTitleColor(const GameMenuLayoutType::Type InLayoutType) { LayoutType = InLayoutType; return *this; }
+
+	/**
+	* Title border margin.
+	*/
+	UPROPERTY(EditAnywhere, Category = Layout)
+	FMargin TitleBorderMargin;
+	FGameMenuStyle& SetTitleBorderMargin(const FMargin InTitleBorderMargin) { TitleBorderMargin = InTitleBorderMargin; return *this; }
+
+	/**
+	* Item border margin.
+	*/
+	UPROPERTY(EditAnywhere, Category = Layout)
+	FMargin ItemBorderMargin;
+	FGameMenuStyle& SetItemBorderMargin(const FMargin InItemBorderMargin) { ItemBorderMargin = InItemBorderMargin; return *this; }
+
+	/**
+	* Title horizontal alignment.
+	*/
+	UPROPERTY(EditAnywhere, Category = Layout)
+	TEnumAsByte<EHorizontalAlignment> TitleHorizontalAlignment;
+	FGameMenuStyle& SetTitleHorizontalAlignment(const EHorizontalAlignment InTitleHorizontalAlignment) { TitleHorizontalAlignment = InTitleHorizontalAlignment; return *this; }
+
+	/**
+	* Title vertical alignment.
+	*/
+	UPROPERTY(EditAnywhere, Category = Layout)
+	TEnumAsByte<EVerticalAlignment> TitleVerticalAlignment;
+	FGameMenuStyle& SetTitleVerticalAlignment(const EVerticalAlignment InTitleVerticalAlignment) { TitleVerticalAlignment = InTitleVerticalAlignment; return *this; }
+
+	/**
+	* Panel vertical alignment.
+	*/
+	UPROPERTY(EditAnywhere, Category = Layout)
+	TEnumAsByte<EVerticalAlignment> PanelsVerticalAlignment;
+	FGameMenuStyle& SetPanelsVerticalAlignment(const EVerticalAlignment InPanelsVerticalAlignment) { PanelsVerticalAlignment = InPanelsVerticalAlignment; return *this; }
+
+	/**
+	 * Speed at which the menu initially appears. 
+	 */
+	UPROPERTY(EditAnywhere, Category = Animation)
+	float AnimationSpeed;
+	FGameMenuStyle& SetMainAnimDuration(const float InAnimationSpeed) { AnimationSpeed = InAnimationSpeed; return *this; }
+	
+	/**
+	* Color of the text.
+	*/
+	UPROPERTY(EditAnywhere, Category = Color)
+	FSlateColor TextColor;
+	FGameMenuStyle& SetTextColor(const FSlateColor InTextColor) { TextColor = InTextColor; return *this; }
+};
+
+
+/**
+*/
+UCLASS(hidecategories = Object, MinimalAPI)
+class UGameMenuWidgetStyle : public USlateWidgetStyleContainerBase
+{
+	GENERATED_UCLASS_BODY()
+
+public:
+	/** The actual data describing the menu's appearance. */
+	UPROPERTY(Category = Appearance, EditAnywhere, meta = (ShowOnlyInnerProperties))
+	FGameMenuStyle MenuStyle;
+
+	virtual const struct FSlateWidgetStyle* const GetStyle() const override
+	{
+		return static_cast< const struct FSlateWidgetStyle* >(&MenuStyle);
+	}
+};
\ No newline at end of file
diff --git a/Engine/Source/Runtime/GameMenuBuilder/Public/SGameMenuItemWidget.h b/Engine/Source/Runtime/GameMenuBuilder/Public/SGameMenuItemWidget.h
new file mode 100644
--- /dev/null
+++ b/Engine/Source/Runtime/GameMenuBuilder/Public/SGameMenuItemWidget.h
@@ -0,0 +1,111 @@
+// Copyright 1998-2014 Epic Games, Inc. All Rights Reserved.
+
+#pragma once
+
+// Menu item widget
+class SGameMenuItemWidget : public SCompoundWidget
+{
+public:
+	DECLARE_DELEGATE_OneParam( FOnArrowPressed, int );
+
+	SLATE_BEGIN_ARGS(SGameMenuItemWidget)
+	{}
+
+	/** style to use for this menu item */
+	SLATE_STYLE_ARGUMENT(FGameMenuStyle, MenuStyle)
+
+	/** weak pointer to the parent PC */
+	SLATE_ARGUMENT(TWeakObjectPtr<APlayerController>, PCOwner)
+
+	/** called when the button is clicked */
+	SLATE_EVENT(FOnClicked, OnClicked)
+
+	/** called when the left or right arrow is clicked */
+	SLATE_EVENT(FOnArrowPressed, OnArrowPressed)
+
+	/** menu item text attribute */
+	SLATE_ATTRIBUTE(FString, Text)
+
+	/** is it multi-choice item? */
+	SLATE_ARGUMENT(bool, bIsMultichoice)
+
+	/** menu item option text attribute */
+	SLATE_ATTRIBUTE(FString, OptionText)
+
+	/** menu item text transparency when item is not active, optional argument */
+	SLATE_ARGUMENT(TOptional<float>, InactiveTextAlpha)
+
+	/** end of slate attributes definition */
+	SLATE_END_ARGS()
+
+	/** needed for every widget */
+	void Construct(const FArguments& InArgs);
+
+	// Begin SWidget interface
+	virtual bool SupportsKeyboardFocus() const override { return true; }
+	
+	virtual FReply OnMouseButtonDown(const FGeometry& MyGeometry, const FPointerEvent& MouseEvent) override;
+
+	virtual FReply OnMouseButtonUp(const FGeometry& MyGeometry, const FPointerEvent& MouseEvent) override;
+
+	virtual FReply OnMouseMove(const FGeometry& MyGeometry, const FPointerEvent& MouseEvent) override;
+	// End SWidget interface
+
+	/** Sets this menu item as active (selected) */
+	void SetMenuItemActive(bool bIsMenuItemActive);
+
+	/** set in option item to enable left arrow*/
+	EVisibility LeftArrowVisible;
+
+	/** set in option item to enable right arrow*/
+	EVisibility RightArrowVisible;
+
+protected:
+	/** Delegate to execute when the button is clicked. */
+	FOnClicked OnClicked;
+
+	/** Delegate to execute when one of arrows was pressed. */
+	FOnArrowPressed OnArrowPressed;
+
+private:
+	/** Menu item text attribute. */
+	TAttribute< FString > Text;
+
+	/** Menu item option text attribute. */
+	TAttribute< FString > OptionText;
+
+	/** Getter for menu item background color. */
+	FSlateColor GetButtonBgColor() const;
+
+	/** Getter for menu item text color. */
+	FSlateColor GetButtonTextColor() const;
+
+	/** Getter for left option arrow visibility. */
+	EVisibility GetLeftArrowVisibility() const;
+
+	/** Getter for right option arrow visibility. */
+	EVisibility GetRightArrowVisibility() const;
+
+	/** Right arrrow pressed delegate. */
+	FReply OnRightArrowDown(const FGeometry& MyGeometry, const FPointerEvent& MouseEvent);
+
+	/** Left Arrow pressed handler. */
+	FReply OnLeftArrowDown(const FGeometry& MyGeometry, const FPointerEvent& MouseEvent);
+
+	/** Inactive text alpha value. */
+	float InactiveTextAlpha;
+
+	/** Active item flag. */
+	bool bIsActiveMenuItem;
+
+	/** Does this menu item represent multi-choice field. */
+	bool bIsMultichoice;
+
+	/** Pointer to our parent Player Controller. */
+	TWeakObjectPtr<class APlayerController> PCOwner;
+
+	/* The Style of the menu. */
+	const FGameMenuStyle* MenuStyle;
+};
+
+
diff --git a/Engine/Source/Runtime/GameMenuBuilder/Public/SGameMenuPageWidget.h b/Engine/Source/Runtime/GameMenuBuilder/Public/SGameMenuPageWidget.h
new file mode 100644
--- /dev/null
+++ b/Engine/Source/Runtime/GameMenuBuilder/Public/SGameMenuPageWidget.h
@@ -0,0 +1,366 @@
+// Copyright 1998-2014 Epic Games, Inc. All Rights Reserved.
+
+#pragma once
+
+class UGameViewportClient;
+
+namespace EPanelState
+{
+	enum Type
+	{
+		/* Menu is opening. */
+		Opening,
+		/* Menu is open. */
+		Open,
+		/* Menu is closing. */
+		Closing,
+		/* Menu is closed. */
+		Closed
+	};
+
+};
+
+DECLARE_DELEGATE_OneParam(FPanelStateChanged, bool);
+
+// Simple class to contain the menu panels/animations
+class FMenuPanel
+{
+public:
+	FMenuPanel();
+
+	void Tick(float Delta);
+
+	/* Close the panel */
+	void ClosePanel();
+	
+	/* Open the panel */
+	void OpenPanel();
+
+	/* Force the panel to be fully open. */
+	void ForcePanelOpen();
+
+	/* Force the panel to be fully closed. */
+	void ForcePanelClosed();
+	
+	void SetStyle(const FGameMenuStyle* InStyle);
+
+	/* Delegate called when the panel becomes open or closed. */
+	FPanelStateChanged	OnStateChanged;
+	
+	/* Animation curve/handle for panel animation. */
+	FCurveHandle AnimationHandle;
+
+	/* The current state of the panel. */
+	EPanelState::Type	CurrentState;
+
+private:		
+	/* Animation sequence used to open or close the panel. */
+	FCurveSequence AnimationSequence;
+};
+
+
+// Simple Menu page widget class
+class GAMEMENUBUILDER_API SGameMenuPageWidget : public SCompoundWidget
+{
+public:
+	SLATE_BEGIN_ARGS(SGameMenuPageWidget)
+	: _PCOwner() {}
+
+	/** style to use for this menu item. */
+	SLATE_STYLE_ARGUMENT(FGameMenuStyle, MenuStyle)
+
+	/** weak pointer to the parent HUD base. */
+	SLATE_ARGUMENT(TWeakObjectPtr<APlayerController>, PCOwner)
+
+	/** is this main menu or in game menu ? */
+	SLATE_ARGUMENT(bool, GameMenu)
+
+	/** always goes here */
+	SLATE_END_ARGS()
+		
+	/** Delegate to call when in-game menu should be hidden using controller buttons - 
+	it's workaround as when joystick is captured, even when sending FReply::Unhandled, binding does not recieve input :( */
+	DECLARE_DELEGATE(FOnToggleMenu);
+	
+	/** Delegate for selection changing. Passes old and new selection. */
+	DECLARE_DELEGATE_TwoParams(FOnSelectionChanged, TSharedPtr<class FGameMenuItem>, TSharedPtr<class FGameMenuItem>);
+
+	/** every widget needs a construction function. */
+	void Construct(const FArguments& InArgs);
+
+	// Begin SWidget interface
+	/** Update function. */
+	virtual void Tick(const FGeometry& AllottedGeometry, const double InCurrentTime, const float InDeltaTime) override;
+
+	/** Mouse button down handler. */
+	virtual FReply OnMouseButtonDown(const FGeometry& MyGeometry, const FPointerEvent& MouseEvent) override;
+
+	/** Key down handler. */
+	virtual FReply OnKeyDown(const FGeometry& MyGeometry, const FKeyboardEvent& InKeyboardEvent) override;
+
+	/** Controller button handler. */
+	virtual FReply OnControllerButtonPressed(const FGeometry& MyGeometry, const FControllerEvent& ControllerEvent) override;
+
+	/** says that we can support keyboard focus. */
+	virtual bool SupportsKeyboardFocus() const override { return true; }
+
+	/** The menu sets up the appropriate mouse settings upon focus. */
+	virtual FReply OnKeyboardFocusReceived(const FGeometry& MyGeometry, const FKeyboardFocusEvent& InKeyboardFocusEvent) override;
+
+	// End SWidget interface
+
+	/**
+	 * Callback handler for when the state of the main panel changes.
+	 *
+	 * @param bWasOpened	whether the panel was opened or closed. is TRUE if panel opened.
+	 *
+	 */ 
+	virtual void OnMainPanelStateChange(bool bWasOpened);
+
+	
+	/**
+ 	 * Callback handler for when the state of the sub panel changes.b 
+	 *
+	 * @param bWasOpened	whether the panel was opened or closed. is TRUE if panel opened.
+	 *
+	 */ 
+	virtual void OnSubPanelStateChange(bool bWasOpened);
+
+	/** Setups misc animation sequences. */
+	void SetupAnimations();
+
+	/** Make the currently selected menu sub menu new main menu if valid. */
+	void EnterSubMenu(TSharedPtr<class FGameMenuPage> InSubMenu);
+
+	/*
+	 * Go back the the previous menu.
+	 * 
+	 * @param	bIsCancel	if true will be treated as CANCEL (IE escape pressed).
+	 *
+	 */
+	void MenuGoBack(bool bIsCancel=false);
+
+	/** Confirms current menu item and performs an action. Will also play selection sound. */
+	void ConfirmMenuItem();
+
+	/*
+	 * Show the given menu and make it the current menu. 
+	 *
+	 * @param InMenu	The menu to show and set as the current menu. If this is NULL Current menu will be used if valid.
+	 *
+	 */
+	void BuildAndShowMenu(TSharedPtr< class FGameMenuPage > InMenu);
+
+	/** Hide the menu. */
+	void HideMenu();
+
+	/*
+	 * Updates arrows visibility for multi-choice menu item.
+	 *
+	 * @param	InMenuItem The item to update.
+	 *
+	 */
+	void UpdateArrows(TSharedPtr<class FGameMenuItem> InMenuItem);
+
+	/*
+	 * Change the currently selection option of the currently selected menu item.
+	 *
+	 * @param	InMoveBy		Number to change the option by.
+	 *
+	 */
+	void ChangeOption(int32 InMoveBy);
+
+	/** disable/enable moving around menu. */
+	void LockControls(bool bEnable);
+
+	/* 
+	 * Set the current menu and 'open' it as the main panel. 
+	 *
+	 * @param InMenu	The menu to open and set as current.
+	 *
+	 */
+	void OpenMainPanel(TSharedPtr< class FGameMenuPage > InMenu);
+		
+	/*
+	* Select a given item from an index
+	*
+	* @param InSelection	The index of the item to select
+	* @returns true if selection changed.
+	*/
+	bool SelectItem(int32 InSelection);
+
+	void ResetMenu();
+
+	/* The viewport I am attached to. */
+	TWeakObjectPtr<UGameViewportClient> MyGameViewport;
+
+protected:
+
+	/* Set the current menu. Also sets the owner widget of that menu to this and resets the previous menu owner. */
+	void SetCurrentMenu(TSharedPtr< class FGameMenuPage > InMenu);
+
+	/** Sets hit test invisibility when console is up. */
+	EVisibility GetSlateVisibility() const;
+
+	/* Returns true if the submenu should be visible. */
+	EVisibility GetSubMenuVisiblity() const;
+
+	/** Gets current scale for drawing menu. */
+	float GetUIScale() const;
+
+	/* 
+	 * Rebuild the widgets in the main menu widget.
+	 * Selected item will be set from InPreviousIndex unless it is invalid.
+	 *
+	 * @param		InPanel		The FGameMenuPage we are building the panel from.
+	 * @param		InBox		The SVerticalBox in which to build the widgets from the given panel.
+	 * @param		PreviousIndex	The previous menu index. (Can be INDEX_NONE)
+	 *
+	 */
+	void BuildPanelButtons(TSharedPtr< class FGameMenuPage > InPanel, TSharedPtr<SVerticalBox> InBox, int32 InPreviousIndex);
+
+	/** Misc Getters used for animating/fading various menu facets. */
+	FLinearColor GetPanelsColor() const;	
+	FSlateColor GetPanelsBackgroundColor() const;
+	FMargin GetMainMenuOffset() const;
+	FMargin GetSubMenuOffset() const;
+	FMargin GetMenuTitleOffset() const;
+	FMargin GetMenuOffset() const;
+	FMargin GetMenuItemPadding() const;
+	FMargin GetSubMenuItemPadding() const;
+	FSlateColor GetTitleColor() const;
+	FSlateColor GetTextColor() const;
+
+	/** 
+	 * Callback for when one of the menu items is selected. 
+	 *
+	 * @param	SelectionIndex	The index selection.
+	 * @returns	Will return UNHANDLED if the current menu is invalid, HANDLED in all other cases including if the controls are locked.
+	 */
+	virtual FReply SelectionChanged(int32 SelectionIndex);
+
+	/** 
+	 * Callback for handling mouse click. 
+	 *
+	 * @param	SelectionIndex	The index selection.
+	 * @returns	Will return UNHANDLED if the current menu is invalid, HANDLED in all other cases including if the controls are locked.
+	 */	
+	virtual FReply MouseButtonClicked(int32 SelectionIndex);
+
+	/** 
+	 * Gets currently selected multi-choice option. 
+	 *
+	 * @param	SelectionIndex	The index selection.
+	 * @returns	Will return UNHANDLED if the current menu is invalid, HANDLED in all other cases including if the controls are locked.
+	 */
+	FString GetOptionText(TSharedPtr<class FGameMenuItem> MenuItem) const;
+
+	/** Gets current menu title string. */
+	FString GetMenuTitle() const; 
+
+	/** Gets the menu title visibility based on if the title text is empty. */
+	EVisibility GetMenuTitleVisibility() const;
+
+	/** Start the main menu open fade/anim and set keyboard focu.s */
+	void FadeIn();
+
+	/* 
+	 * Animates the title widget.
+	 *
+	 * @param	bShowTitle	true to show the title, false to hide it.
+	 */
+	void SetTitleAnimation(bool bShowTitle);
+
+	/* Opens any pending sub menu if there is one. */
+	void OpenPendingSubMenu();
+
+	/** Control for hiding the menu - Defaults to start button. */
+	//FKey ControllerHideMenuKey;
+	
+	/** Current menu title if present. */
+	FText CurrentMenuTitle;
+
+	/** If console is currently opened. */
+	bool bConsoleVisible;
+	
+	/** Container instance for main panel. */
+	FMenuPanel	MainMenuPanel;
+
+	/** Container instance for sub menu panel (only relevant with side by side layout). */
+	FMenuPanel	SubMenuPanel;
+	
+	/** Next menu (for transition and displaying as the right menu). */
+	TSharedPtr< class FGameMenuPage > NextMenu;
+
+	/** Currently active menu. */
+	TSharedPtr< class FGameMenuPage > CurrentMenu;
+
+	/** Current UI scale attribute. */
+	TAttribute<float> UIScale;
+
+	/** Our curve sequence and the related handles. */
+	FCurveSequence MenuWidgetAnimation;
+
+	/** Used for main menu logo fade in animation at the beginning.  */
+	FCurveHandle TopColorCurve;
+
+	/** Used for menu background fade in animation at the beginning. */
+	FCurveHandle ColorFadeCurve;
+
+	/** Used for menu buttons slide in animation at the beginning. */
+	FCurveHandle MenuAnimationCurve;
+
+	/** Our curve sequence and the related handles. */
+	FCurveSequence TitleWidgetAnimation;
+
+	/* Used to animate the title widget (NYI) */
+	FCurveHandle TitleWidgetCurve;
+
+	/** Weak pointer to our parent Player Controller. */
+	TWeakObjectPtr<class APlayerController> PCOwner;
+
+	/** Screen resolution. */
+	FIntPoint ScreenRes;
+
+	/** Animation type index. */
+	int32 MainAnimNumber;
+
+	/** Selected index of current menu. */
+	int32 SelectedIndex;
+	
+	/** Flag when playing hiding animation. */
+	uint8 bMenuHiding : 1;
+
+	/** Flag when playing hiding animation. */
+	uint8 bMenuHidden : 1;
+
+ 	/** If this is in game menu. */
+ 	uint8 bGameMenu : 1;
+
+	/** If moving around menu is currently locked. */
+	uint8 bControlsLocked : 1;
+	
+	/** Menu that will override current one after transition animation. */
+	TSharedPtr< class FGameMenuPage > PendingMainMenu;
+	
+	/** Menu that will become the submenu after transition animation. */
+	TSharedPtr< class FGameMenuPage > PendingSubMenu;
+	
+	/** Current menu layout box. */
+	TSharedPtr<SVerticalBox> MainPanel;
+
+	/** Sub menu layout box. */
+	TSharedPtr<SVerticalBox> SubPanel;
+
+	/** Style to use for this menu item. */
+	const FGameMenuStyle* MenuStyle;
+		
+	/** Stack of previous menus. */
+	TArray< TSharedPtr< class FGameMenuPage>> MenuHistory;
+
+	/** Bind if menu should be hidden from outside by controller button. */
+	FOnToggleMenu OnToggleMenu;
+	
+	/** Executed when the user selects an item. */
+	FOnSelectionChanged OnSelectionChange;
+};
