Converted GameMenuBuilder to exclusively use FText with Slate

UETOOL-213 - Minimize Slate FString -> FText conversion (remove SLATE_TEXT_ATTRIBUTE)

diff --git a/Engine/Source/Runtime/GameMenuBuilder/Private/SGameMenuItemWidget.cpp b/Engine/Source/Runtime/GameMenuBuilder/Private/SGameMenuItemWidget.cpp
--- a/Engine/Source/Runtime/GameMenuBuilder/Private/SGameMenuItemWidget.cpp
+++ b/Engine/Source/Runtime/GameMenuBuilder/Private/SGameMenuItemWidget.cpp
@@ -63,7 +63,7 @@ void SGameMenuItemWidget::Construct(const FArguments& InArgs)
 					SNew(STextBlock)
 					.TextStyle(FGameMenuBuilderStyle::Get(), "GameMenuStyle.MenuTextStyle")
 					.ColorAndOpacity(this,&SGameMenuItemWidget::GetButtonTextColor)
-					.Text(FString("<"))
+					.Text(FText::FromString(TEXT("<")))
 				]
 			]
 			+SHorizontalBox::Slot()
@@ -87,7 +87,7 @@ void SGameMenuItemWidget::Construct(const FArguments& InArgs)
 					SNew(STextBlock)
 					.TextStyle(FGameMenuBuilderStyle::Get(), "GameMenuStyle.MenuTextStyle")
 					.ColorAndOpacity(this,&SGameMenuItemWidget::GetButtonTextColor)
-					.Text(FString(">"))
+					.Text(FText::FromString(TEXT(">")))
 				]
 			]
 		]
diff --git a/Engine/Source/Runtime/GameMenuBuilder/Private/SGameMenuPageWidget.cpp b/Engine/Source/Runtime/GameMenuBuilder/Private/SGameMenuPageWidget.cpp
--- a/Engine/Source/Runtime/GameMenuBuilder/Private/SGameMenuPageWidget.cpp
+++ b/Engine/Source/Runtime/GameMenuBuilder/Private/SGameMenuPageWidget.cpp
@@ -383,7 +383,7 @@ void SGameMenuPageWidget::BuildPanelButtons(TSharedPtr< class FGameMenuPage > In
 					.MenuStyle(MenuStyle)
 					.PCOwner(PCOwner)
 					.OnClicked(this, &SGameMenuPageWidget::MouseButtonClicked, i)
-					.Text(EachItem.Get()->Text.ToString())
+					.Text(EachItem.Get()->Text)
 					.bIsMultichoice(false);
 			}
 			else if (EachItem->MenuItemType == EGameMenuItemType::MultiChoice)
@@ -392,7 +392,7 @@ void SGameMenuPageWidget::BuildPanelButtons(TSharedPtr< class FGameMenuPage > In
 					.MenuStyle(MenuStyle)
 					.PCOwner(PCOwner)
 					.OnClicked(this, &SGameMenuPageWidget::MouseButtonClicked, i)
-					.Text(EachItem.Get()->Text.ToString())
+					.Text(EachItem.Get()->Text)
 					.bIsMultichoice(true)
 					.OnArrowPressed(this, &SGameMenuPageWidget::ChangeOption)
 					.OptionText(this, &SGameMenuPageWidget::GetOptionText, EachItem);
@@ -427,9 +427,9 @@ void SGameMenuPageWidget::BuildPanelButtons(TSharedPtr< class FGameMenuPage > In
 
 }
 
-FString SGameMenuPageWidget::GetOptionText(TSharedPtr<FGameMenuItem> InMenuItem) const
+FText SGameMenuPageWidget::GetOptionText(TSharedPtr<FGameMenuItem> InMenuItem) const
 {
-	return InMenuItem->MultiChoice[InMenuItem->SelectedMultiChoice].ToString();
+	return InMenuItem->MultiChoice[InMenuItem->SelectedMultiChoice];
 }
 
 void SGameMenuPageWidget::UpdateArrows(TSharedPtr<FGameMenuItem> InMenuItem)
diff --git a/Engine/Source/Runtime/GameMenuBuilder/Public/SGameMenuItemWidget.h b/Engine/Source/Runtime/GameMenuBuilder/Public/SGameMenuItemWidget.h
--- a/Engine/Source/Runtime/GameMenuBuilder/Public/SGameMenuItemWidget.h
+++ b/Engine/Source/Runtime/GameMenuBuilder/Public/SGameMenuItemWidget.h
@@ -24,13 +24,13 @@ public:
 	SLATE_EVENT(FOnArrowPressed, OnArrowPressed)
 
 	/** menu item text attribute */
-	SLATE_ATTRIBUTE(FString, Text)
+	SLATE_ATTRIBUTE(FText, Text)
 
 	/** is it multi-choice item? */
 	SLATE_ARGUMENT(bool, bIsMultichoice)
 
 	/** menu item option text attribute */
-	SLATE_ATTRIBUTE(FString, OptionText)
+	SLATE_ATTRIBUTE(FText, OptionText)
 
 	/** menu item text transparency when item is not active, optional argument */
 	SLATE_ARGUMENT(TOptional<float>, InactiveTextAlpha)
@@ -69,10 +69,10 @@ protected:
 
 private:
 	/** Menu item text attribute. */
-	TAttribute< FString > Text;
+	TAttribute< FText > Text;
 
 	/** Menu item option text attribute. */
-	TAttribute< FString > OptionText;
+	TAttribute< FText > OptionText;
 
 	/** Getter for menu item background color. */
 	FSlateColor GetButtonBgColor() const;
diff --git a/Engine/Source/Runtime/GameMenuBuilder/Public/SGameMenuPageWidget.h b/Engine/Source/Runtime/GameMenuBuilder/Public/SGameMenuPageWidget.h
--- a/Engine/Source/Runtime/GameMenuBuilder/Public/SGameMenuPageWidget.h
+++ b/Engine/Source/Runtime/GameMenuBuilder/Public/SGameMenuPageWidget.h
@@ -250,7 +250,7 @@ protected:
 	 * @param	SelectionIndex	The index selection.
 	 * @returns	Will return UNHANDLED if the current menu is invalid, HANDLED in all other cases including if the controls are locked.
 	 */
-	FString GetOptionText(TSharedPtr<class FGameMenuItem> MenuItem) const;
+	FText GetOptionText(TSharedPtr<class FGameMenuItem> MenuItem) const;
 
 	/** Gets current menu title string. */
 	FText GetMenuTitle() const; 
