Correct initialization of bMenuHiding and bMenuHidden in SGameMenuPageWidget::Construct(). Fixes issues with clicking on player name input widget in PlatformerGame.

diff --git a/Engine/Source/Runtime/GameMenuBuilder/Private/SGameMenuPageWidget.cpp b/Engine/Source/Runtime/GameMenuBuilder/Private/SGameMenuPageWidget.cpp
--- a/Engine/Source/Runtime/GameMenuBuilder/Private/SGameMenuPageWidget.cpp
+++ b/Engine/Source/Runtime/GameMenuBuilder/Private/SGameMenuPageWidget.cpp
@@ -136,7 +136,8 @@ void SGameMenuPageWidget::Construct(const FArguments& InArgs)
 	
 	bControlsLocked = false;
 	bConsoleVisible = false;
-	bMenuHidden = false;
+	bMenuHiding = false;
+	bMenuHidden = true;
 	SelectedIndex = INDEX_NONE;
 	
 
@@ -313,7 +314,7 @@ bool SGameMenuPageWidget::SelectItem(int32 InSelection)
 void SGameMenuPageWidget::BuildAndShowMenu(TSharedPtr< class FGameMenuPage > InMenu)
 {
 	SetCurrentMenu(InMenu);
-	
+
 	bMenuHiding = false;
 	bMenuHidden = false;
 		
