Some minor revisions to GameMenuBuilder
Removed some old dead code in PlatformerGameModule

diff --git a/Engine/Source/Runtime/GameMenuBuilder/Private/GameMenuBuilderStyle.cpp b/Engine/Source/Runtime/GameMenuBuilder/Private/GameMenuBuilderStyle.cpp
--- a/Engine/Source/Runtime/GameMenuBuilder/Private/GameMenuBuilderStyle.cpp
+++ b/Engine/Source/Runtime/GameMenuBuilder/Private/GameMenuBuilderStyle.cpp
@@ -38,19 +38,22 @@ FName FGameMenuBuilderStyle::GetStyleSetName()
 #define TTF_FONT( RelativePath, ... ) FSlateFontInfo( FPaths::GameContentDir() / "Slate"/ RelativePath + TEXT(".ttf"), __VA_ARGS__ )
 #define OTF_FONT( RelativePath, ... ) FSlateFontInfo( FPaths::GameContentDir() / "Slate"/ RelativePath + TEXT(".otf"), __VA_ARGS__ )
 
+FString FGameMenuBuilderStyle::FontName("Fonts/Roboto-Light");
+int32 FGameMenuBuilderStyle::FontSize = 42;
+
 TSharedRef< FSlateStyleSet > FGameMenuBuilderStyle::Create(const FString StyleName)
 {
-	TSharedRef<FSlateGameResources> StyleRef = FSlateGameResources::New(FGameMenuBuilderStyle::GetStyleSetName(), *StyleName);
+	TSharedRef<FSlateGameResources> StyleRef = FSlateGameResources::New(GetStyleSetName(), *StyleName);
 	
 	FSlateStyleSet& Style = StyleRef.Get();
 
 	// Fonts still need to be specified in code for now
 	Style.Set("GameMenuStyle.MenuTextStyle", FTextBlockStyle()
-		.SetFont(TTF_FONT("Fonts/Roboto-Light", 42))
+		.SetFont(TTF_FONT(*FontName, FontSize))
 		.SetColorAndOpacity(FLinearColor::White)
 		);
 	Style.Set("GameMenuStyle.MenuHeaderTextStyle", FTextBlockStyle()
-		.SetFont(TTF_FONT("Fonts/Roboto-Light", 42))
+		.SetFont(TTF_FONT(*FontName, FontSize))
 		.SetColorAndOpacity(FLinearColor::White)
 		);
 
diff --git a/Engine/Source/Runtime/GameMenuBuilder/Private/SGameMenuItemWidget.cpp b/Engine/Source/Runtime/GameMenuBuilder/Private/SGameMenuItemWidget.cpp
--- a/Engine/Source/Runtime/GameMenuBuilder/Private/SGameMenuItemWidget.cpp
+++ b/Engine/Source/Runtime/GameMenuBuilder/Private/SGameMenuItemWidget.cpp
@@ -182,3 +182,23 @@ void SGameMenuItemWidget::SetMenuItemActive(bool bIsMenuItemActive)
 {
 	this->bIsActiveMenuItem = bIsMenuItemActive;
 }
+
+void SGameMenuItemWidget::SetMenuOwner(TWeakObjectPtr<class APlayerController> InPCOwner)
+{
+	PCOwner = InPCOwner;
+}
+
+void SGameMenuItemWidget::SetMenuStyle(const FGameMenuStyle* InMenuStyle)
+{
+	MenuStyle = InMenuStyle;
+}
+
+void SGameMenuItemWidget::SetClickedDelegate(FOnClicked InOnClicked)
+{
+	OnClicked = InOnClicked;
+}
+
+void SGameMenuItemWidget::SetArrowPressedDelegate(FOnArrowPressed InOnArrowPressed)
+{
+	OnArrowPressed = InOnArrowPressed;
+}
diff --git a/Engine/Source/Runtime/GameMenuBuilder/Private/SGameMenuPageWidget.cpp b/Engine/Source/Runtime/GameMenuBuilder/Private/SGameMenuPageWidget.cpp
--- a/Engine/Source/Runtime/GameMenuBuilder/Private/SGameMenuPageWidget.cpp
+++ b/Engine/Source/Runtime/GameMenuBuilder/Private/SGameMenuPageWidget.cpp
@@ -120,6 +120,7 @@ void SGameMenuPageWidget::Construct(const FArguments& InArgs)
 	MenuStyle = InArgs._MenuStyle;
 	check(MenuStyle);
 	MainMenuPanel.SetStyle(MenuStyle);
+	SubMenuPanel.SetStyle(MenuStyle);
 	
 	EHorizontalAlignment MainAlignmentH;
 	if (MenuStyle->LayoutType == GameMenuLayoutType::Single)
@@ -376,7 +377,7 @@ void SGameMenuPageWidget::BuildPanelButtons(TSharedPtr< class FGameMenuPage > In
 	{
 		for (int32 i = 0; i < InPanel->NumItems(); ++i)
 		{
-			TSharedPtr<SWidget> TmpWidget;
+			TSharedPtr<SGameMenuItemWidget> TmpWidget;
 			TSharedPtr<FGameMenuItem> EachItem = InPanel->GetItem(i);
 			if (EachItem.Get()->MenuItemType == EGameMenuItemType::Standard)
 			{
@@ -402,6 +403,8 @@ void SGameMenuPageWidget::BuildPanelButtons(TSharedPtr< class FGameMenuPage > In
 			else if (EachItem.Get()->MenuItemType == EGameMenuItemType::CustomWidget)
 			{
 				TmpWidget = EachItem.Get()->CustomWidget;
+				TmpWidget->SetMenuOwner(PCOwner);
+				TmpWidget->SetMenuStyle(MenuStyle);
 			}
 			InBox->AddSlot().HAlign(HAlign_Left).AutoHeight()
 				[
@@ -586,11 +589,7 @@ void SGameMenuPageWidget::OnMainPanelStateChange(bool bWasOpened)
 		if (PendingMainMenu.IsValid())
 		{
 			OpenMainPanel(PendingMainMenu);
-			if (MenuStyle->LayoutType == GameMenuLayoutType::Single)
-			{
-				OpenMainPanel(PendingMainMenu);
-			}
-			else
+			if (MenuStyle->LayoutType != GameMenuLayoutType::Single)
 			{
 				MainMenuPanel.ForcePanelOpen();
 				SubMenuPanel.ForcePanelClosed();				
@@ -1057,3 +1056,9 @@ void SGameMenuPageWidget::ResetMenu()
 	CurrentMenu.Reset();
 	SubPanel.Reset();	
 }
+
+
+TSharedPtr< FGameMenuPage > SGameMenuPageWidget::GetCurrentMenu()
+{
+	return CurrentMenu;
+}
diff --git a/Engine/Source/Runtime/GameMenuBuilder/Public/GameMenuBuilderStyle.h b/Engine/Source/Runtime/GameMenuBuilder/Public/GameMenuBuilderStyle.h
--- a/Engine/Source/Runtime/GameMenuBuilder/Public/GameMenuBuilderStyle.h
+++ b/Engine/Source/Runtime/GameMenuBuilder/Public/GameMenuBuilderStyle.h
@@ -18,6 +18,8 @@ public:
 
 	static FName GetStyleSetName();
 		
+	GAMEMENUBUILDER_API static FString FontName;
+	GAMEMENUBUILDER_API static int32	FontSize;
 private:
 
 	static TSharedRef< class FSlateStyleSet > Create(const FString StyleName);
diff --git a/Engine/Source/Runtime/GameMenuBuilder/Public/GameMenuItem.h b/Engine/Source/Runtime/GameMenuBuilder/Public/GameMenuItem.h
--- a/Engine/Source/Runtime/GameMenuBuilder/Public/GameMenuItem.h
+++ b/Engine/Source/Runtime/GameMenuBuilder/Public/GameMenuItem.h
@@ -40,7 +40,7 @@ public:
 	TSharedPtr<SGameMenuItemWidget> Widget;
 
 	/** shared pointer to actual slate widget representing the custom menu item, ie whole options screen */
-	TSharedPtr<SWidget> CustomWidget;
+	TSharedPtr<SGameMenuItemWidget> CustomWidget;
 
 	/** texts for multiple choice menu item (like INF AMMO ON/OFF or difficulty/resolution etc) */
 	TArray<FText> MultiChoice;
@@ -67,10 +67,10 @@ public:
 	}
 
 	/** custom widgets cannot contain sub menus, all functionality must be handled by custom widget itself */
-	FGameMenuItem(TSharedPtr<SWidget> _Widget)
+	FGameMenuItem(TSharedPtr<SGameMenuItemWidget> InWidget)
 	{
 		MenuItemType = EGameMenuItemType::CustomWidget;
-		CustomWidget = _Widget;
+		CustomWidget = InWidget;
 		MinMultiChoiceIndex = MaxMultiChoiceIndex = SelectedMultiChoice = 0;
 		bInactive = false;
 	}
diff --git a/Engine/Source/Runtime/GameMenuBuilder/Public/GameMenuPage.h b/Engine/Source/Runtime/GameMenuBuilder/Public/GameMenuPage.h
--- a/Engine/Source/Runtime/GameMenuBuilder/Public/GameMenuPage.h
+++ b/Engine/Source/Runtime/GameMenuBuilder/Public/GameMenuPage.h
@@ -93,6 +93,24 @@ public:
 		return Item.ToSharedRef();
 	}
 
+	/**
+	 * Add a custom menu entry 
+	 *
+	 * @param	Text		Label of the menu item
+	 * @param	
+	 * @param	
+	 * @param	
+	 * @returns  SharedRef to the MenuItem that was created.
+	 */
+	template< class UserClass >
+	FORCEINLINE TSharedRef<FGameMenuItem> AddCustomMenuItem(const FText& Text, TSharedPtr<SGameMenuItemWidget> CustomWidget, UserClass* InObj, typename FGameMenuItem::FOnOptionChanged::TSPMethodDelegate< UserClass >::FMethodPtr InMethod)
+	{
+		TSharedPtr<FGameMenuItem> Item = MakeShareable(new FGameMenuItem(CustomWidget));
+		Item->OnOptionChanged.BindSP(InObj, InMethod);
+		MenuItems.Add(Item);
+		return Item.ToSharedRef();
+	}
+
 	/**
 	* Add a handler for the menu being canceled.
 	*
diff --git a/Engine/Source/Runtime/GameMenuBuilder/Public/SGameMenuItemWidget.h b/Engine/Source/Runtime/GameMenuBuilder/Public/SGameMenuItemWidget.h
--- a/Engine/Source/Runtime/GameMenuBuilder/Public/SGameMenuItemWidget.h
+++ b/Engine/Source/Runtime/GameMenuBuilder/Public/SGameMenuItemWidget.h
@@ -3,7 +3,7 @@
 #pragma once
 
 // Menu item widget
-class SGameMenuItemWidget : public SCompoundWidget
+class GAMEMENUBUILDER_API SGameMenuItemWidget : public SCompoundWidget
 {
 public:
 	DECLARE_DELEGATE_OneParam( FOnArrowPressed, int );
@@ -60,6 +60,17 @@ public:
 	/** set in option item to enable right arrow*/
 	EVisibility RightArrowVisible;
 
+	/** Set pointer to our parent Player Controller. */
+	void SetMenuOwner(TWeakObjectPtr<class APlayerController> InPCOwner);
+	
+	/** Set pointer to our style. */
+	void SetMenuStyle(const FGameMenuStyle* InMenuStyle);
+
+	void SetClickedDelegate(FOnClicked InOnClicked);
+
+	/** Delegate to execute when one of arrows was pressed. */
+	void SetArrowPressedDelegate(FOnArrowPressed InOnArrowPressed);
+
 protected:
 	/** Delegate to execute when the button is clicked. */
 	FOnClicked OnClicked;
@@ -67,13 +78,29 @@ protected:
 	/** Delegate to execute when one of arrows was pressed. */
 	FOnArrowPressed OnArrowPressed;
 
-private:
 	/** Menu item text attribute. */
 	TAttribute< FText > Text;
 
 	/** Menu item option text attribute. */
 	TAttribute< FText > OptionText;
 
+	/** Inactive text alpha value. */
+	float InactiveTextAlpha;
+
+	/** Active item flag. */
+	bool bIsActiveMenuItem;
+
+	/** Does this menu item represent multi-choice field. */
+	bool bIsMultichoice;
+
+	/** Pointer to our parent Player Controller. */
+	TWeakObjectPtr<class APlayerController> PCOwner;
+
+	/* The Style of the menu. */
+	const FGameMenuStyle* MenuStyle;
+
+private:
+	
 	/** Getter for menu item background color. */
 	FSlateColor GetButtonBgColor() const;
 
@@ -92,20 +119,6 @@ private:
 	/** Left Arrow pressed handler. */
 	FReply OnLeftArrowDown(const FGeometry& MyGeometry, const FPointerEvent& MouseEvent);
 
-	/** Inactive text alpha value. */
-	float InactiveTextAlpha;
-
-	/** Active item flag. */
-	bool bIsActiveMenuItem;
-
-	/** Does this menu item represent multi-choice field. */
-	bool bIsMultichoice;
-
-	/** Pointer to our parent Player Controller. */
-	TWeakObjectPtr<class APlayerController> PCOwner;
-
-	/* The Style of the menu. */
-	const FGameMenuStyle* MenuStyle;
 };
 
 
diff --git a/Engine/Source/Runtime/GameMenuBuilder/Public/SGameMenuPageWidget.h b/Engine/Source/Runtime/GameMenuBuilder/Public/SGameMenuPageWidget.h
--- a/Engine/Source/Runtime/GameMenuBuilder/Public/SGameMenuPageWidget.h
+++ b/Engine/Source/Runtime/GameMenuBuilder/Public/SGameMenuPageWidget.h
@@ -58,6 +58,7 @@ private:
 	FCurveSequence AnimationSequence;
 };
 
+class FGameMenuPage;
 
 // Simple Menu page widget class
 class GAMEMENUBUILDER_API SGameMenuPageWidget : public SCompoundWidget
@@ -127,7 +128,7 @@ public:
 	void SetupAnimations();
 
 	/** Make the currently selected menu sub menu new main menu if valid. */
-	void EnterSubMenu(TSharedPtr<class FGameMenuPage> InSubMenu);
+	void EnterSubMenu(TSharedPtr<FGameMenuPage> InSubMenu);
 
 	/**
 	 * Go back the the previous menu.
@@ -146,7 +147,7 @@ public:
 	 * @param InMenu	The menu to show and set as the current menu. If this is NULL Current menu will be used if valid.
 	 *
 	 */
-	void BuildAndShowMenu(TSharedPtr< class FGameMenuPage > InMenu);
+	void BuildAndShowMenu(TSharedPtr< FGameMenuPage > InMenu);
 
 	/** Hide the menu. */
 	void HideMenu();
@@ -176,7 +177,7 @@ public:
 	 * @param InMenu	The menu to open and set as current.
 	 *
 	 */
-	void OpenMainPanel(TSharedPtr< class FGameMenuPage > InMenu);
+	void OpenMainPanel(TSharedPtr< FGameMenuPage > InMenu);
 		
 	/**
 	* Select a given item from an index
@@ -191,10 +192,11 @@ public:
 	/* The viewport I am attached to. */
 	TWeakObjectPtr<UGameViewportClient> MyGameViewport;
 
+	TSharedPtr< FGameMenuPage > GetCurrentMenu();
 protected:
 
 	/* Set the current menu. Also sets the owner widget of that menu to this and resets the previous menu owner. */
-	void SetCurrentMenu(TSharedPtr< class FGameMenuPage > InMenu);
+	void SetCurrentMenu(TSharedPtr< FGameMenuPage > InMenu);
 
 	/** Sets hit test invisibility when console is up. */
 	EVisibility GetSlateVisibility() const;
@@ -214,7 +216,7 @@ protected:
 	 * @param		PreviousIndex	The previous menu index. (Can be INDEX_NONE)
 	 *
 	 */
-	void BuildPanelButtons(TSharedPtr< class FGameMenuPage > InPanel, TSharedPtr<SVerticalBox> InBox, int32 InPreviousIndex);
+	void BuildPanelButtons(TSharedPtr< FGameMenuPage > InPanel, TSharedPtr<SVerticalBox> InBox, int32 InPreviousIndex);
 
 	/** Misc Getters used for animating/fading various menu facets. */
 	FLinearColor GetPanelsColor() const;	
@@ -287,10 +289,10 @@ protected:
 	FMenuPanel	SubMenuPanel;
 	
 	/** Next menu (for transition and displaying as the right menu). */
-	TSharedPtr< class FGameMenuPage > NextMenu;
+	TSharedPtr< FGameMenuPage > NextMenu;
 
 	/** Currently active menu. */
-	TSharedPtr< class FGameMenuPage > CurrentMenu;
+	TSharedPtr< FGameMenuPage > CurrentMenu;
 
 	/** Current UI scale attribute. */
 	TAttribute<float> UIScale;
@@ -338,10 +340,10 @@ protected:
 	uint8 bControlsLocked : 1;
 	
 	/** Menu that will override current one after transition animation. */
-	TSharedPtr< class FGameMenuPage > PendingMainMenu;
+	TSharedPtr< FGameMenuPage > PendingMainMenu;
 	
 	/** Menu that will become the submenu after transition animation. */
-	TSharedPtr< class FGameMenuPage > PendingSubMenu;
+	TSharedPtr< FGameMenuPage > PendingSubMenu;
 	
 	/** Current menu layout box. */
 	TSharedPtr<SVerticalBox> MainPanel;
@@ -353,7 +355,7 @@ protected:
 	const FGameMenuStyle* MenuStyle;
 		
 	/** Stack of previous menus. */
-	TArray< TSharedPtr< class FGameMenuPage>> MenuHistory;
+	TArray< TSharedPtr< FGameMenuPage>> MenuHistory;
 
 	/** Bind if menu should be hidden from outside by controller button. */
 	FOnToggleMenu OnToggleMenu;
