Fix PlatformerGame shadow variables

diff --git a/Engine/Source/Runtime/GameMenuBuilder/Private/SGameMenuPageWidget.cpp b/Engine/Source/Runtime/GameMenuBuilder/Private/SGameMenuPageWidget.cpp
--- a/Engine/Source/Runtime/GameMenuBuilder/Private/SGameMenuPageWidget.cpp
+++ b/Engine/Source/Runtime/GameMenuBuilder/Private/SGameMenuPageWidget.cpp
@@ -838,7 +838,6 @@ FReply SGameMenuPageWidget::SelectionChanged(int32 SelectionIndex)
 	if (SelectedIndex != SelectionIndex)
 	{
 		TSharedPtr<FGameMenuItem> CurrentMenuItem;
-		TSharedPtr<SGameMenuItemWidget> MenuWidgetItem;
 		if (SelectedIndex != INDEX_NONE)
 		{
 			CurrentMenuItem = CurrentMenu->GetItem(SelectedIndex);
