[GameMenuBuilder]
* Fixed bad include. Should include GameViewportClient.h instead of Viewport.h

diff --git a/Engine/Source/Runtime/GameMenuBuilder/Private/SGameMenuPageWidget.cpp b/Engine/Source/Runtime/GameMenuBuilder/Private/SGameMenuPageWidget.cpp
--- a/Engine/Source/Runtime/GameMenuBuilder/Private/SGameMenuPageWidget.cpp
+++ b/Engine/Source/Runtime/GameMenuBuilder/Private/SGameMenuPageWidget.cpp
@@ -1,7 +1,7 @@
 // Copyright Epic Games, Inc. All Rights Reserved.
 
 #include "SGameMenuPageWidget.h"
-#include "Components/Viewport.h"
+#include "Engine/GameViewportClient.h"
 #include "Widgets/SBoxPanel.h"
 #include "Widgets/SOverlay.h"
 #include "Framework/Application/SlateApplication.h"
