Replace some usages of FORCEINLINE with inline in UI modules.

For context about this change, see this doc: https://docs.google.com/document/d/1LjpT7NpzTecfXdFNyv5KM6sCu67827BRljNi1iaWVSE/edit?tab=t.0#heading=h.fvtz8gjukz7x

Reviewers were automatically determined. Please comment if anyone else should be added.

#rnx

diff --git a/Engine/Plugins/Experimental/UIFramework/Source/Public/Types/UIFWidgetId.h b/Engine/Plugins/Experimental/UIFramework/Source/Public/Types/UIFWidgetId.h
--- a/Engine/Plugins/Experimental/UIFramework/Source/Public/Types/UIFWidgetId.h
+++ b/Engine/Plugins/Experimental/UIFramework/Source/Public/Types/UIFWidgetId.h
@@ -53,7 +53,7 @@ public:
 		return Key != Other.Key;
 	}
 	
-	FORCEINLINE friend uint32 GetTypeHash(const FUIFrameworkWidgetId& Value)
+	inline friend uint32 GetTypeHash(const FUIFrameworkWidgetId& Value)
 	{
 		return GetTypeHash(Value.Key);
 	}
diff --git a/Engine/Plugins/Runtime/CommonUI/Source/CommonUI/Public/Slate/SCommonAnimatedSwitcher.h b/Engine/Plugins/Runtime/CommonUI/Source/CommonUI/Public/Slate/SCommonAnimatedSwitcher.h
--- a/Engine/Plugins/Runtime/CommonUI/Source/CommonUI/Public/Slate/SCommonAnimatedSwitcher.h
+++ b/Engine/Plugins/Runtime/CommonUI/Source/CommonUI/Public/Slate/SCommonAnimatedSwitcher.h
@@ -57,7 +57,7 @@ enum class ECommonSwitcherTransitionFallbackStrategy : uint8
 	Last
 };
 
-static FORCEINLINE ECurveEaseFunction TransitionCurveToCurveEaseFunction(ETransitionCurve CurveType)
+static inline ECurveEaseFunction TransitionCurveToCurveEaseFunction(ETransitionCurve CurveType)
 {
 	switch (CurveType)
 	{
diff --git a/Engine/Plugins/Runtime/CommonUI/Source/CommonUI/Public/UITag.h b/Engine/Plugins/Runtime/CommonUI/Source/CommonUI/Public/UITag.h
--- a/Engine/Plugins/Runtime/CommonUI/Source/CommonUI/Public/UITag.h
+++ b/Engine/Plugins/Runtime/CommonUI/Source/CommonUI/Public/UITag.h
@@ -45,7 +45,7 @@ class TTypedTagStaticImpl
 
 #if !(UE_BUILD_SHIPPING || UE_BUILD_TEST)
 	/** Intended for console commands/cheats: not for shipping code! */
-	static FORCEINLINE TagT FindFromString_DebugOnly(const FString& PartialTagName)
+	static inline TagT FindFromString_DebugOnly(const FString& PartialTagName)
 	{
 		return UGameplayTagsManager::Get().FindGameplayTagFromPartialString_Slow(PartialTagName);
 	}
@@ -162,7 +162,7 @@ struct FGlobalUITags : public FGameplayTagNativeAdder
 		UIAction_NextTab = FUIActionTag::AddNativeTag(TEXT("NextTab"));
 	}
 
-	FORCEINLINE static const FGlobalUITags& Get() { return GUITags; }
+	inline static const FGlobalUITags& Get() { return GUITags; }
 
 private:
 	static UE_API FGlobalUITags GUITags;
diff --git a/Engine/Plugins/Runtime/CommonUI/Source/CommonUI/Public/WidgetFactory.h b/Engine/Plugins/Runtime/CommonUI/Source/CommonUI/Public/WidgetFactory.h
--- a/Engine/Plugins/Runtime/CommonUI/Source/CommonUI/Public/WidgetFactory.h
+++ b/Engine/Plugins/Runtime/CommonUI/Source/CommonUI/Public/WidgetFactory.h
@@ -15,9 +15,9 @@ class TWidgetFactory : public FGCObject
 {
 	static_assert(TIsDerivedFrom<WidgetType, UUserWidget>::IsDerived, "Can only use TWidgetFactory with types deriving from UUserWidget");
 private:
-	FORCEINLINE UGameInstance* OwningObjectAsGameInstance() const { return Cast<UGameInstance>(OuterGetter()); }
-	FORCEINLINE UWorld* OwningObjectAsWorld() const { return Cast<UWorld>(OuterGetter()); }
-	FORCEINLINE APlayerController* OwningObjectAsPlayerController() const { return Cast<APlayerController>(OuterGetter()); }
+	inline UGameInstance* OwningObjectAsGameInstance() const { return Cast<UGameInstance>(OuterGetter()); }
+	inline UWorld* OwningObjectAsWorld() const { return Cast<UWorld>(OuterGetter()); }
+	inline APlayerController* OwningObjectAsPlayerController() const { return Cast<APlayerController>(OuterGetter()); }
 
 public:
 	typedef TFunctionRef<TSharedPtr<SObjectWidget>(UUserWidget*, TSharedRef<SWidget>)> ConstructMethodType;
diff --git a/Engine/Plugins/VirtualProduction/RemoteControl/Source/RemoteControlUI/Public/IRemoteControlUIModule.h b/Engine/Plugins/VirtualProduction/RemoteControl/Source/RemoteControlUI/Public/IRemoteControlUIModule.h
--- a/Engine/Plugins/VirtualProduction/RemoteControl/Source/RemoteControlUI/Public/IRemoteControlUIModule.h
+++ b/Engine/Plugins/VirtualProduction/RemoteControl/Source/RemoteControlUI/Public/IRemoteControlUIModule.h
@@ -84,17 +84,17 @@ struct FRCExposesPropertyArgs
 	/** Checked version of Get Property */
 	FProperty* GetPropertyChecked() const;
 
-	friend FORCEINLINE bool operator==(const FRCExposesPropertyArgs& LHS, const FRCExposesPropertyArgs& RHS)
+	friend inline bool operator==(const FRCExposesPropertyArgs& LHS, const FRCExposesPropertyArgs& RHS)
 	{
 		return LHS.Id == RHS.Id;
 	}
 
-	friend FORCEINLINE bool operator!=(const FRCExposesPropertyArgs& LHS, const FRCExposesPropertyArgs& RHS)
+	friend inline bool operator!=(const FRCExposesPropertyArgs& LHS, const FRCExposesPropertyArgs& RHS)
 	{
 		return LHS.Id != RHS.Id;
 	}
 
-	friend FORCEINLINE uint32 GetTypeHash(const FRCExposesPropertyArgs& InArgs)
+	friend inline uint32 GetTypeHash(const FRCExposesPropertyArgs& InArgs)
 	{
 		return GetTypeHash(InArgs.Id);
 	}
diff --git a/Engine/Source/Runtime/Experimental/FieldSystem/Source/FieldSystemEngine/Public/Field/FieldSystemComponent.h b/Engine/Source/Runtime/Experimental/FieldSystem/Source/FieldSystemEngine/Public/Field/FieldSystemComponent.h
--- a/Engine/Source/Runtime/Experimental/FieldSystem/Source/FieldSystemEngine/Public/Field/FieldSystemComponent.h
+++ b/Engine/Source/Runtime/Experimental/FieldSystem/Source/FieldSystemEngine/Public/Field/FieldSystemComponent.h
@@ -52,7 +52,7 @@ public:
 	void SetFieldSystem(UFieldSystem* FieldSystemIn) { FieldSystem = FieldSystemIn; }
 
 	/** Get the field system asset */
-	FORCEINLINE const UFieldSystem* GetFieldSystem() const { return FieldSystem; }
+	inline const UFieldSystem* GetFieldSystem() const { return FieldSystem; }
 
 	/** Field system asset to be used to store the construction fields */
 	UPROPERTY(EditAnywhere, AdvancedDisplay, BlueprintReadOnly, Category = "Field", meta = (ToolTip = "Field system asset to be used to store the construction fields. This asset is not required anymore and will be deprecated soon."))
diff --git a/Engine/Source/Runtime/FieldNotification/Public/FieldNotificationHelpers.h b/Engine/Source/Runtime/FieldNotification/Public/FieldNotificationHelpers.h
--- a/Engine/Source/Runtime/FieldNotification/Public/FieldNotificationHelpers.h
+++ b/Engine/Source/Runtime/FieldNotification/Public/FieldNotificationHelpers.h
@@ -7,7 +7,7 @@
 
 namespace UE::FieldNotification::Helpers
 {
-	[[nodiscard]] FORCEINLINE bool IsValidAsField(const UFunction* InFunction)
+	[[nodiscard]] inline bool IsValidAsField(const UFunction* InFunction)
 	{
 		return InFunction != nullptr
 			&& !InFunction->HasAnyFunctionFlags(FUNC_Net | FUNC_Event)
diff --git a/Engine/Source/Runtime/GameMenuBuilder/Public/GameMenuPage.h b/Engine/Source/Runtime/GameMenuBuilder/Public/GameMenuPage.h
--- a/Engine/Source/Runtime/GameMenuBuilder/Public/GameMenuPage.h
+++ b/Engine/Source/Runtime/GameMenuBuilder/Public/GameMenuPage.h
@@ -78,7 +78,7 @@ public:
 	* @returns  SharedRef to the MenuItem that was created.
 	*/
 	template< class UserClass >
-	FORCEINLINE TSharedRef<FGameMenuItem> AddMenuItem(const FText& InText, UserClass* InObj, typename FGameMenuItem::FOnConfirmMenuItem::TMethodPtr< UserClass > InMethod)
+	inline TSharedRef<FGameMenuItem> AddMenuItem(const FText& InText, UserClass* InObj, typename FGameMenuItem::FOnConfirmMenuItem::TMethodPtr< UserClass > InMethod)
 	{
 		TSharedPtr<FGameMenuItem> Item = MakeShareable(new FGameMenuItem(InText));
 		Item->OnConfirmMenuItem.BindSP(InObj, InMethod);
@@ -96,7 +96,7 @@ public:
 	 * @returns  SharedRef to the MenuItem that was created.
 	 */
 	template< class UserClass >
-	FORCEINLINE TSharedRef<FGameMenuItem> AddMenuItemWithOptions(const FText& Text, const TArray<FText>& OptionsList, UserClass* InObj, typename FGameMenuItem::FOnOptionChanged::TMethodPtr< UserClass > InMethod)
+	inline TSharedRef<FGameMenuItem> AddMenuItemWithOptions(const FText& Text, const TArray<FText>& OptionsList, UserClass* InObj, typename FGameMenuItem::FOnOptionChanged::TMethodPtr< UserClass > InMethod)
 	{
 		TSharedPtr<FGameMenuItem> Item = MakeShareable(new FGameMenuItem(Text, OptionsList));
 		Item->OnOptionChanged.BindSP(InObj, InMethod);
@@ -114,7 +114,7 @@ public:
 	 * @returns  SharedRef to the MenuItem that was created.
 	 */
 	template< class UserClass >
-	FORCEINLINE TSharedRef<FGameMenuItem> AddCustomMenuItem(const FText& Text, TSharedPtr<SGameMenuItemWidget> CustomWidget, UserClass* InObj, typename FGameMenuItem::FOnOptionChanged::TMethodPtr< UserClass > InMethod)
+	inline TSharedRef<FGameMenuItem> AddCustomMenuItem(const FText& Text, TSharedPtr<SGameMenuItemWidget> CustomWidget, UserClass* InObj, typename FGameMenuItem::FOnOptionChanged::TMethodPtr< UserClass > InMethod)
 	{
 		TSharedPtr<FGameMenuItem> Item = MakeShareable(new FGameMenuItem(CustomWidget));
 		Item->OnOptionChanged.BindSP(InObj, InMethod);
@@ -129,7 +129,7 @@ public:
 	* @param	InMethod	Method to call when selection changes.
 	*/
 	template< class UserClass >
-	FORCEINLINE void SetCancelHandler(UserClass* InObj, typename FGameMenuPage::FOnMenuGoBack::TMethodPtr< UserClass > InMethod)
+	inline void SetCancelHandler(UserClass* InObj, typename FGameMenuPage::FOnMenuGoBack::TMethodPtr< UserClass > InMethod)
 	{
 		OnGoBackCancel.BindSP(InObj, InMethod);
 	}
@@ -141,7 +141,7 @@ public:
 	* @param	InMethod	Method to call when menu has been hidden.
 	*/
 	template< class UserClass >
-	FORCEINLINE void SetOnHiddenHandler(UserClass* InObj, typename FGameMenuPage::FOnMenuHidden::TMethodPtr< UserClass > InMethod)
+	inline void SetOnHiddenHandler(UserClass* InObj, typename FGameMenuPage::FOnMenuHidden::TMethodPtr< UserClass > InMethod)
 	{
 		OnMenuHidden.BindSP(InObj, InMethod);
 	}
@@ -153,7 +153,7 @@ public:
 	* @param	InMethod	Method to call when selection changes.
 	*/
 	template< class UserClass >
-	FORCEINLINE void SetAcceptHandler(UserClass* InObj, typename FGameMenuPage::FOnMenuGoBack::TMethodPtr< UserClass > InMethod)
+	inline void SetAcceptHandler(UserClass* InObj, typename FGameMenuPage::FOnMenuGoBack::TMethodPtr< UserClass > InMethod)
 	{
 		OnGoBack.BindSP(InObj, InMethod);
 	}
@@ -165,7 +165,7 @@ public:
 	* @param	InMethod	Method to call when menu is about to open.
 	*/
 	template< class UserClass >
-	FORCEINLINE void SetOnOpenHandler(UserClass* InObj, typename FGameMenuPage::FOnMenuHidden::TMethodPtr< UserClass > InMethod)
+	inline void SetOnOpenHandler(UserClass* InObj, typename FGameMenuPage::FOnMenuHidden::TMethodPtr< UserClass > InMethod)
 	{
 		OnMenuOpening.BindSP(InObj, InMethod);
 	}
diff --git a/Engine/Source/Runtime/UMG/Public/Blueprint/UserWidget.h b/Engine/Source/Runtime/UMG/Public/Blueprint/UserWidget.h
--- a/Engine/Source/Runtime/UMG/Public/Blueprint/UserWidget.h
+++ b/Engine/Source/Runtime/UMG/Public/Blueprint/UserWidget.h
@@ -1579,7 +1579,7 @@ protected:
 	 */
 	UMG_API virtual int32 NativePaint(const FPaintArgs& Args, const FGeometry& AllottedGeometry, const FSlateRect& MyCullingRect, FSlateWindowElementList& OutDrawElements, int32 LayerId, const FWidgetStyle& InWidgetStyle, bool bParentEnabled ) const;
 
-	FORCEINLINE FVector2D GetMinimumDesiredSize() const { return MinimumDesiredSize; }
+	inline FVector2D GetMinimumDesiredSize() const { return MinimumDesiredSize; }
 	UMG_API void SetMinimumDesiredSize(FVector2D InMinimumDesiredSize);
 
 	UMG_API virtual bool NativeIsInteractable() const;
diff --git a/Engine/Source/Runtime/UMG/Public/Blueprint/WidgetTree.h b/Engine/Source/Runtime/UMG/Public/Blueprint/WidgetTree.h
--- a/Engine/Source/Runtime/UMG/Public/Blueprint/WidgetTree.h
+++ b/Engine/Source/Runtime/UMG/Public/Blueprint/WidgetTree.h
@@ -34,7 +34,7 @@ public:
 
 	/** Finds the widget in the tree by name and casts the return to the desired type. */
 	template <typename WidgetT>
-	FORCEINLINE WidgetT* FindWidget(const FName& Name) const
+	inline WidgetT* FindWidget(const FName& Name) const
 	{
 		return Cast<WidgetT>(FindWidget(Name));
 	}
@@ -88,7 +88,7 @@ public:
 
 	/** Constructs the widget, and adds it to the tree. */
 	template <typename WidgetT>
-	[[nodiscard]] FORCEINLINE_DEBUGGABLE WidgetT* ConstructWidget(TSubclassOf<WidgetT> WidgetClass = WidgetT::StaticClass(), FName WidgetName = NAME_None)
+	[[nodiscard]] inline WidgetT* ConstructWidget(TSubclassOf<WidgetT> WidgetClass = WidgetT::StaticClass(), FName WidgetName = NAME_None)
 	{
 		if(WidgetClass != nullptr)
 		{
diff --git a/Engine/Source/Runtime/UMG/Public/Components/PanelSlot.h b/Engine/Source/Runtime/UMG/Public/Components/PanelSlot.h
--- a/Engine/Source/Runtime/UMG/Public/Components/PanelSlot.h
+++ b/Engine/Source/Runtime/UMG/Public/Components/PanelSlot.h
@@ -27,7 +27,7 @@ public:
 #if WITH_EDITOR
 	UMG_API bool IsDesignTime() const;
 #else
-	FORCEINLINE bool IsDesignTime() const { return false; }
+	inline bool IsDesignTime() const { return false; }
 #endif
 
 	UMG_API virtual void ReleaseSlateResources(bool bReleaseChildren) override;
diff --git a/Engine/Source/Runtime/UMG/Public/Components/Viewport.h b/Engine/Source/Runtime/UMG/Public/Components/Viewport.h
--- a/Engine/Source/Runtime/UMG/Public/Components/Viewport.h
+++ b/Engine/Source/Runtime/UMG/Public/Components/Viewport.h
@@ -52,16 +52,16 @@ public:
 
 
 	/** @return The transform's location */
-	FORCEINLINE const FVector& GetLocation() const { return ViewLocation; }
+	inline const FVector& GetLocation() const { return ViewLocation; }
 
 	/** @return The transform's rotation */
-	FORCEINLINE const FRotator& GetRotation() const { return ViewRotation; }
+	inline const FRotator& GetRotation() const { return ViewRotation; }
 
 	/** @return The look at point for orbiting */
-	FORCEINLINE const FVector& GetLookAt() const { return LookAt; }
+	inline const FVector& GetLookAt() const { return LookAt; }
 
 	/** @return The ortho zoom amount */
-	FORCEINLINE float GetOrthoZoom() const { return OrthoZoom; }
+	inline float GetOrthoZoom() const { return OrthoZoom; }
 
 	/**
 	* Animates from the current location to the desired location
diff --git a/Engine/Source/Runtime/UMG/Public/Components/Widget.h b/Engine/Source/Runtime/UMG/Public/Components/Widget.h
--- a/Engine/Source/Runtime/UMG/Public/Components/Widget.h
+++ b/Engine/Source/Runtime/UMG/Public/Components/Widget.h
@@ -927,7 +927,7 @@ public:
 
 #if WITH_EDITOR
 	/** Returns if the widget is currently being displayed in the designer, it may want to display different data. */
-	FORCEINLINE bool IsDesignTime() const
+	inline bool IsDesignTime() const
 	{
 		return HasAnyDesignerFlags(EWidgetDesignFlags::Designing);
 	}
@@ -936,18 +936,18 @@ public:
 	UMG_API virtual void SetDesignerFlags(EWidgetDesignFlags NewFlags);
 
 	/** Gets the designer flags currently set on the widget. */
-	FORCEINLINE EWidgetDesignFlags GetDesignerFlags() const
+	inline EWidgetDesignFlags GetDesignerFlags() const
 	{
 		return static_cast<EWidgetDesignFlags>(DesignerFlags);
 	}
 
 	/** Tests if any of the flags exist on this widget. */
-	FORCEINLINE bool HasAnyDesignerFlags(EWidgetDesignFlags FlagsToCheck) const
+	inline bool HasAnyDesignerFlags(EWidgetDesignFlags FlagsToCheck) const
 	{
 		return EnumHasAnyFlags(GetDesignerFlags(), FlagsToCheck);
 	}
 	
-	FORCEINLINE bool IsPreviewTime() const 
+	inline bool IsPreviewTime() const 
 	{
 		return HasAnyDesignerFlags(EWidgetDesignFlags::Previewing);
 	}
@@ -977,8 +977,8 @@ public:
 	/** Mark this object as modified, also mark the slot as modified. */
 	UMG_API virtual bool Modify(bool bAlwaysMarkDirty = true) override;
 #else
-	FORCEINLINE bool IsDesignTime() const { return false; }
-	FORCEINLINE bool IsPreviewTime() const { return false; }
+	inline bool IsDesignTime() const { return false; }
+	inline bool IsPreviewTime() const { return false; }
 #endif
 	
 	/**
@@ -1016,18 +1016,18 @@ public:
 	UMG_API virtual void PreSave(FObjectPreSaveContext ObjectSaveContext) override;
 	//~ End UObject
 
-	FORCEINLINE bool CanSafelyRouteEvent()
+	inline bool CanSafelyRouteEvent()
 	{
 		return !IsDesignTime() && CanSafelyRouteCall();
 	}
 
-	FORCEINLINE bool CanSafelyRoutePaint()
+	inline bool CanSafelyRoutePaint()
 	{
 		return CanSafelyRouteCall();
 	}
 
 protected:
-	FORCEINLINE bool CanSafelyRouteCall() { return !(GIntraFrameDebuggingGameThread || IsUnreachable() || FUObjectThreadContext::Get().IsRoutingPostLoad); }
+	inline bool CanSafelyRouteCall() { return !(GIntraFrameDebuggingGameThread || IsUnreachable() || FUObjectThreadContext::Get().IsRoutingPostLoad); }
 
 public:
 
@@ -1221,7 +1221,7 @@ private:
 	bool bRoutedSynchronizeProperties;
 
 #else
-	FORCEINLINE void VerifySynchronizeProperties() { }
+	inline void VerifySynchronizeProperties() { }
 #endif
 
 private:
diff --git a/Engine/Source/Runtime/UMG/Public/Slate/SMeshWidget.h b/Engine/Source/Runtime/UMG/Public/Slate/SMeshWidget.h
--- a/Engine/Source/Runtime/UMG/Public/Slate/SMeshWidget.h
+++ b/Engine/Source/Runtime/UMG/Public/Slate/SMeshWidget.h
@@ -62,7 +62,7 @@ public:
 	 * @param InInstanceOffset   Start drawing with this instance
 	 * @param InNumInstances     Draw this many instances
 	 */
-	FORCEINLINE void AddRenderRun(uint32 InMeshIndex, uint32 InInstanceOffset, uint32 InNumInstances)
+	inline void AddRenderRun(uint32 InMeshIndex, uint32 InInstanceOffset, uint32 InNumInstances)
 	{
 		RenderRuns.Add(FRenderRun(InMeshIndex, InInstanceOffset, InNumInstances));
 	}
diff --git a/Engine/Source/Runtime/UMG/Public/Slate/SObjectWidget.h b/Engine/Source/Runtime/UMG/Public/Slate/SObjectWidget.h
--- a/Engine/Source/Runtime/UMG/Public/Slate/SObjectWidget.h
+++ b/Engine/Source/Runtime/UMG/Public/Slate/SObjectWidget.h
@@ -108,12 +108,12 @@ protected:
 	TObjectPtr<UUserWidget> WidgetObject;
 
 private:
-	FORCEINLINE bool CanRouteEvent() const
+	inline bool CanRouteEvent() const
 	{
 		return WidgetObject && WidgetObject->CanSafelyRouteEvent();
 	}
 
-	FORCEINLINE bool CanRoutePaint() const
+	inline bool CanRoutePaint() const
 	{
 		return WidgetObject && WidgetObject->CanSafelyRoutePaint();
 	}
diff --git a/Engine/Source/Runtime/UMG/Public/Slate/WidgetTransform.h b/Engine/Source/Runtime/UMG/Public/Slate/WidgetTransform.h
--- a/Engine/Source/Runtime/UMG/Public/Slate/WidgetTransform.h
+++ b/Engine/Source/Runtime/UMG/Public/Slate/WidgetTransform.h
@@ -59,17 +59,17 @@ public:
 		return Identity == *this;
 	}
 
-	FORCEINLINE bool operator==( const FWidgetTransform &Other ) const
+	inline bool operator==( const FWidgetTransform &Other ) const
 	{
 		return Scale == Other.Scale && Shear == Other.Shear && Angle == Other.Angle && Translation == Other.Translation;
 	}
 
-	FORCEINLINE bool operator!=( const FWidgetTransform& Other ) const
+	inline bool operator!=( const FWidgetTransform& Other ) const
 	{
 		return !( *this == Other );
 	}
 
-	FORCEINLINE FSlateRenderTransform ToSlateRenderTransform() const
+	inline FSlateRenderTransform ToSlateRenderTransform() const
 	{
 		return ::Concatenate(
 			FScale2D(Scale),
